UEENT-3174
DMX plugin arhitecture upload
1. All plugins setup
2. DMX broadcasting
3. DMX recieving raw data

diff --git a/Engine/Plugins/Experimental/DigitalMultiplex/DMXProtocol/DMXProtocol.uplugin b/Engine/Plugins/Experimental/DigitalMultiplex/DMXProtocol/DMXProtocol.uplugin
new file mode 100644
--- /dev/null
+++ b/Engine/Plugins/Experimental/DigitalMultiplex/DMXProtocol/DMXProtocol.uplugin
@@ -0,0 +1,32 @@
+{
+  "FileVersion": 1,
+  "Version": 1,
+  "VersionName": "1.0",
+  "FriendlyName": "DMX Protocol",
+  "Description": "DMX Protocols implementation",
+  "Category": "Media",
+  "CreatedBy": "Epic Games, Inc.",
+  "CreatedByURL": "http://epicgames.com",
+  "DocsURL": "",
+  "MarketplaceURL": "",
+  "SupportURL": "",
+  "CanContainContent": true,
+  "IsBetaVersion": true,
+  "Modules": [
+    {
+      "Name": "DMXProtocol",
+      "Type": "Runtime",
+      "LoadingPhase": "Default"
+    },
+    {
+      "Name": "DMXProtocolArtNet",
+      "Type": "Runtime",
+      "LoadingPhase": "Default"
+    },
+    {
+      "Name": "DMXProtocolSACN",
+      "Type": "Runtime",
+      "LoadingPhase": "Default"
+    }
+  ]
+}
diff --git a/Engine/Plugins/Experimental/DigitalMultiplex/DMXProtocol/Source/DMXProtocol/DMXProtocol.Build.cs b/Engine/Plugins/Experimental/DigitalMultiplex/DMXProtocol/Source/DMXProtocol/DMXProtocol.Build.cs
new file mode 100644
--- /dev/null
+++ b/Engine/Plugins/Experimental/DigitalMultiplex/DMXProtocol/Source/DMXProtocol/DMXProtocol.Build.cs
@@ -0,0 +1,32 @@
+// Copyright 1998-2019 Epic Games, Inc. All Rights Reserved.
+
+using UnrealBuildTool;
+
+public class DMXProtocol : ModuleRules
+{
+	public DMXProtocol(ReadOnlyTargetRules Target) : base(Target)
+	{
+		PCHUsage = ModuleRules.PCHUsageMode.UseExplicitOrSharedPCHs;
+		
+		PublicDependencyModuleNames.AddRange(
+			new string[]
+			{
+				"Core",
+            }
+			);
+			
+		
+		PrivateDependencyModuleNames.AddRange(
+			new string[]
+			{
+				"CoreUObject",
+				"Engine",
+                "Serialization",
+                "Networking",
+                "Sockets",
+                "Json"
+            }
+			);
+		
+	}
+}
diff --git a/Engine/Plugins/Experimental/DigitalMultiplex/DMXProtocol/Source/DMXProtocol/Private/DMXProtocol.cpp b/Engine/Plugins/Experimental/DigitalMultiplex/DMXProtocol/Source/DMXProtocol/Private/DMXProtocol.cpp
new file mode 100644
--- /dev/null
+++ b/Engine/Plugins/Experimental/DigitalMultiplex/DMXProtocol/Source/DMXProtocol/Private/DMXProtocol.cpp
@@ -0,0 +1,79 @@
+// Copyright 1998-2019 Epic Games, Inc. All Rights Reserved.
+
+#include "DMXProtocol.h"
+
+#include "Managers/DMXProtocolDeviceManager.h"
+#include "Managers/DMXProtocolInterfaceManager.h"
+#include "Managers/DMXProtocolPortManager.h"
+#include "Managers/DMXProtocolUniverseManager.h"
+
+#include "Interfaces/IDMXProtocolPort.h"
+#include "Interfaces/IDMXProtocolUniverse.h"
+
+#include "DMXProtocolTypes.h"
+#include "Json.h"
+
+#include "HAL/IConsoleManager.h"
+#include "Logging/LogMacros.h"
+
+DEFINE_LOG_CATEGORY(LogDMXProtocol);
+
+FDMXProtocol::FDMXProtocol(const FName& InProtocolName, FJsonObject& InSettings)
+{
+	Settings = MakeShared<FJsonObject>(InSettings);
+	InterfaceManager = MakeShared<FDMXProtocolInterfaceManager>(this);
+	DeviceManager = MakeShared<FDMXProtocolDeviceManager>(this);
+	PortManager = MakeShared<FDMXProtocolPortManager>(this);
+	UniverseManager = MakeShared<FDMXProtocolUniverseManager>(this);
+}
+
+FDMXProtocol::~FDMXProtocol()
+{
+	InterfaceManager->RemoveAll();
+	DeviceManager->RemoveAll();
+	PortManager->RemoveAll();
+	UniverseManager->RemoveAll();
+}
+
+TSharedPtr<FDMXProtocolDeviceManager> FDMXProtocol::GetDeviceManager() const
+{
+	return DeviceManager;
+}
+
+TSharedPtr<FDMXProtocolInterfaceManager> FDMXProtocol::GetInterfaceManager() const
+{
+	return InterfaceManager;
+}
+
+TSharedPtr<FDMXProtocolPortManager> FDMXProtocol::GetPortManager() const
+{
+	return PortManager;
+}
+
+TSharedPtr<FDMXProtocolUniverseManager> FDMXProtocol::GetUniverseManager() const
+{
+	return UniverseManager;
+}
+
+TSharedPtr<FJsonObject> FDMXProtocol::GetSettings() const
+{
+	return Settings;
+}
+
+void FDMXProtocol::SetDMXFragment(uint16 UniverseID, const IDMXFragmentMap& DMXFragment, bool bShouldSend)
+{
+	IDMXProtocolUniverse* Universe = UniverseManager->GetUniverseById(UniverseID);
+	if (Universe != nullptr)
+	{
+		// Update Universe DMX buffer
+		Universe->SetDMXFragment(DMXFragment);
+		if (bShouldSend)
+		{
+			TSharedPtr<IDMXProtocolPort> Port = Universe->GetCachedUniversePort().Pin();
+			if (Port.IsValid())
+			{
+				Port->WriteDMX(Universe->GetOutputDMXBuffer());
+			}
+		}
+	}
+}
diff --git a/Engine/Plugins/Experimental/DigitalMultiplex/DMXProtocol/Source/DMXProtocol/Private/DMXProtocolModule.cpp b/Engine/Plugins/Experimental/DigitalMultiplex/DMXProtocol/Source/DMXProtocol/Private/DMXProtocolModule.cpp
new file mode 100644
--- /dev/null
+++ b/Engine/Plugins/Experimental/DigitalMultiplex/DMXProtocol/Source/DMXProtocol/Private/DMXProtocolModule.cpp
@@ -0,0 +1,133 @@
+// Copyright 1998-2019 Epic Games, Inc. All Rights Reserved.
+
+#include "DMXProtocolModule.h"
+#include "Interfaces/IDMXProtocol.h"
+#include "Interfaces/IDMXProtocolFactory.h"
+
+IMPLEMENT_MODULE( FDMXProtocolModule, DMXProtocol );
+
+#define LOCTEXT_NAMESPACE "DMXProtocolModule"
+
+const TCHAR* FDMXProtocolModule::BaseModuleName = TEXT("DMXProtocol");
+
+void FDMXProtocolModule::ShutdownModule()
+{
+	ShutdownAllDMXProtocols();
+}
+
+static FName GetProtocolModuleName(const FString& ProtocolName)
+{
+	FName ModuleName;
+	if (!ProtocolName.StartsWith(FDMXProtocolModule::BaseModuleName, ESearchCase::CaseSensitive))
+	{
+		ModuleName = FName(*(FDMXProtocolModule::BaseModuleName + ProtocolName));
+	}
+	else
+	{
+		ModuleName = FName(*ProtocolName);
+	}
+
+	return ModuleName;
+}
+
+static IModuleInterface* GetProtocolModule(const FString& SubsystemName)
+{
+	const FName ModuleName = GetProtocolModuleName(SubsystemName);
+	FModuleManager& ModuleManager = FModuleManager::Get();
+	check(ModuleManager.IsModuleLoaded(ModuleName));
+	return ModuleManager.GetModule(ModuleName);
+}
+
+void FDMXProtocolModule::RegisterProtocol(const FName& FactoryName, IDMXProtocolFactory* Factory)
+{
+	if (!DMXFactories.Contains(FactoryName))
+	{
+		DMXFactories.Add(FactoryName, Factory);
+	}
+	else
+	{
+		UE_LOG_DMXPROTOCOL(Warning, TEXT("Trying to add existing protocol %s"), *FactoryName.ToString());
+	}
+}
+
+void FDMXProtocolModule::UnregisterProtocol(const FName& FactoryName)
+{
+	if (DMXFactories.Contains(FactoryName))
+	{
+		DMXFactories.Remove(FactoryName);
+		ShutdownDMXProtocol(FactoryName);
+	}
+	else
+	{
+		UE_LOG_DMXPROTOCOL(Warning, TEXT("Trying to remove unexisting protocol %s"), *FactoryName.ToString());
+	}
+}
+
+FDMXProtocolModule& FDMXProtocolModule::Get()
+{
+	return FModuleManager::GetModuleChecked<FDMXProtocolModule>("DMXProtocol");
+}
+
+
+IDMXProtocol* FDMXProtocolModule::GetProtocol(const FName InProtocolName)
+{
+	TSharedPtr<IDMXProtocol>* DMXProtocolPtr = nullptr;
+	if (!InProtocolName.IsNone())
+	{
+		DMXProtocolPtr = DMXProtocols.Find(InProtocolName);
+		if (DMXProtocolPtr == nullptr)
+		{
+			IDMXProtocolFactory** DMXProtocolFactory = DMXFactories.Find(InProtocolName);
+
+			if (DMXProtocolFactory != nullptr)
+			{
+				UE_LOG_DMXPROTOCOL(Log, TEXT("Creating protocol instance for: %s"), *InProtocolName.ToString());
+
+				TSharedPtr<IDMXProtocol> NewProtocol = (*DMXProtocolFactory)->CreateProtocol(InProtocolName);
+				if (NewProtocol.IsValid())
+				{
+					DMXProtocols.Add(InProtocolName, NewProtocol);
+					DMXProtocolPtr = DMXProtocols.Find(InProtocolName);
+				}
+				else
+				{
+					bool* bNotedPreviously = DMXProtocolFailureNotes.Find(InProtocolName);
+					if (!bNotedPreviously || !(*bNotedPreviously))
+					{
+						UE_LOG_DMXPROTOCOL(Warning, TEXT("Unable to create Protocol %s"), *InProtocolName.ToString());
+						DMXProtocolFailureNotes.Add(InProtocolName, true);
+					}
+				}
+			}
+		}
+	}
+
+	return (DMXProtocolPtr == nullptr) ? nullptr : (*DMXProtocolPtr).Get();
+}
+
+void FDMXProtocolModule::ShutdownDMXProtocol(const FName& ProtocolName)
+{
+	if (!ProtocolName.IsNone())
+	{
+		TSharedPtr<IDMXProtocol> DMXProtocol;
+		DMXProtocols.RemoveAndCopyValue(ProtocolName, DMXProtocol);
+		if (DMXProtocol.IsValid())
+		{
+			DMXProtocol->Shutdown();
+		}
+		else
+		{
+			UE_LOG_DMXPROTOCOL(Warning, TEXT("DMXProtocol instance %s not found, unable to destroy."), *ProtocolName.ToString());
+		}
+	}
+}
+
+void FDMXProtocolModule::ShutdownAllDMXProtocols()
+{
+	for (TMap<FName, TSharedPtr<IDMXProtocol>>::TIterator It = DMXProtocols.CreateIterator(); It; ++It)
+	{
+		It->Value->Shutdown();
+	}
+}
+
+#undef LOCTEXT_NAMESPACE
diff --git a/Engine/Plugins/Experimental/DigitalMultiplex/DMXProtocol/Source/DMXProtocol/Private/DMXProtocolPackager.cpp b/Engine/Plugins/Experimental/DigitalMultiplex/DMXProtocol/Source/DMXProtocol/Private/DMXProtocolPackager.cpp
new file mode 100644
--- /dev/null
+++ b/Engine/Plugins/Experimental/DigitalMultiplex/DMXProtocol/Source/DMXProtocol/Private/DMXProtocolPackager.cpp
@@ -0,0 +1,26 @@
+// Copyright 1998-2019 Epic Games, Inc. All Rights Reserved.
+
+#include "DMXProtocolPackager.h"
+#include "Interfaces/IDMXProtocolPacket.h"
+#include "Serialization/BufferArchive.h"
+
+bool FDMXProtocolPackager::AddToPackage(IDMXProtocolPacket* Packet)
+{
+	if (Packet != nullptr)
+	{
+		Buffer.Append(*Packet->Pack());
+		return true;
+	}
+
+	return false;
+}
+
+uint32 FDMXProtocolPackager::GetSize() const
+{
+	return Buffer.Num();
+}
+
+const TArray<uint8>& FDMXProtocolPackager::GetBuffer() const
+{
+	return Buffer;
+}
diff --git a/Engine/Plugins/Experimental/DigitalMultiplex/DMXProtocol/Source/DMXProtocol/Private/DMXProtocolTypes.cpp b/Engine/Plugins/Experimental/DigitalMultiplex/DMXProtocol/Source/DMXProtocol/Private/DMXProtocolTypes.cpp
new file mode 100644
--- /dev/null
+++ b/Engine/Plugins/Experimental/DigitalMultiplex/DMXProtocol/Source/DMXProtocol/Private/DMXProtocolTypes.cpp
@@ -0,0 +1,16 @@
+// Copyright 1998-2019 Epic Games, Inc. All Rights Reserved.
+
+#include "DMXProtocolTypes.h"
+
+bool FDMXBuffer::SetDMXFragment(const IDMXFragmentMap & DMXFragment)
+{
+	for (TMap<uint32, uint8>::TConstIterator It = DMXFragment.CreateConstIterator(); It; ++It)
+	{
+		if (It->Key < (DMX_UNIVERSE_SIZE - 1))
+		{
+			DMXData[It->Key] = It->Value;
+		}
+	}
+
+	return false;
+}
diff --git a/Engine/Plugins/Experimental/DigitalMultiplex/DMXProtocol/Source/DMXProtocol/Private/Managers/DMXProtocolDeviceManager.cpp b/Engine/Plugins/Experimental/DigitalMultiplex/DMXProtocol/Source/DMXProtocol/Private/Managers/DMXProtocolDeviceManager.cpp
new file mode 100644
--- /dev/null
+++ b/Engine/Plugins/Experimental/DigitalMultiplex/DMXProtocol/Source/DMXProtocol/Private/Managers/DMXProtocolDeviceManager.cpp
@@ -0,0 +1,52 @@
+// Copyright 1998-2019 Epic Games, Inc. All Rights Reserved.
+
+#include "Managers/DMXProtocolDeviceManager.h"
+#include "Interfaces/IDMXProtocolDevice.h"
+#include "Interfaces/IDMXProtocol.h"
+
+void FDMXProtocolDeviceManager::AddDevice(TSharedPtr<IDMXProtocolDevice> InDevice)
+{
+	checkf(InDevice.IsValid(), TEXT("IDMXProtocolDevicePtr not valid. Manager should own only valid pointer."));
+	if (TSharedPtr<IDMXProtocolInterface> Interface = InDevice->GetCachedProtocolInterface().Pin())
+	{
+		DevicesMap.Add(Interface.Get(), InDevice);
+	}
+	else
+	{
+		UE_LOG_DMXPROTOCOL(Warning, TEXT("Cached interface pointer not valid"));
+	}
+}
+
+void FDMXProtocolDeviceManager::RemoveDevice(TSharedPtr<IDMXProtocolDevice> InDevice)
+{
+	if (TSharedPtr<IDMXProtocolInterface> Interface = InDevice->GetCachedProtocolInterface().Pin())
+	{
+		DevicesMap.Remove(Interface.Get());
+	}
+}
+
+void FDMXProtocolDeviceManager::RemoveAll()
+{
+	DevicesMap.Empty();
+}
+
+bool FDMXProtocolDeviceManager::GetDevicesByProtocol(const FName & ProtocolName, IDevicesMap & OutDevicesMap)
+{
+	bool bIsExists = false;
+
+	for (TMap<IDMXProtocolInterface*, TSharedPtr<IDMXProtocolDevice>>::TConstIterator DeviceIt = DevicesMap.CreateConstIterator(); DeviceIt; ++DeviceIt)
+	{
+		if (DeviceIt->Value.IsValid() && DeviceIt->Value->GetProtocol())
+		{
+			if (DeviceIt->Value->GetProtocol()->GetProtocolName() == ProtocolName)
+			{
+				OutDevicesMap.Add(DeviceIt->Key, DeviceIt->Value);
+				bIsExists = true;
+			}
+		}
+	}
+
+	return bIsExists;
+}
+
+
diff --git a/Engine/Plugins/Experimental/DigitalMultiplex/DMXProtocol/Source/DMXProtocol/Private/Managers/DMXProtocolInterfaceManager.cpp b/Engine/Plugins/Experimental/DigitalMultiplex/DMXProtocol/Source/DMXProtocol/Private/Managers/DMXProtocolInterfaceManager.cpp
new file mode 100644
--- /dev/null
+++ b/Engine/Plugins/Experimental/DigitalMultiplex/DMXProtocol/Source/DMXProtocol/Private/Managers/DMXProtocolInterfaceManager.cpp
@@ -0,0 +1,21 @@
+// Copyright 1998-2019 Epic Games, Inc. All Rights Reserved.
+
+#include "Managers/DMXProtocolInterfaceManager.h"
+#include "Interfaces/IDMXProtocolInterface.h"
+#include "Interfaces/IDMXProtocol.h"
+
+void FDMXProtocolInterfaceManager::AddInterface(int32 InInterfaceID, TSharedPtr<IDMXProtocolInterface> InInterface)
+{
+	checkf(InInterface.IsValid(), TEXT("IDMXProtocolInterfacePtr not valid. Manager should own only valid pointer."));
+	InterfacesMap.Add(InInterfaceID, InInterface);
+}
+
+void FDMXProtocolInterfaceManager::RemoveInterface(int32 InInterfaceID)
+{
+	InterfacesMap.Remove(InInterfaceID);
+}
+
+void FDMXProtocolInterfaceManager::RemoveAll()
+{
+	InterfacesMap.Empty();
+}
diff --git a/Engine/Plugins/Experimental/DigitalMultiplex/DMXProtocol/Source/DMXProtocol/Private/Managers/DMXProtocolPortManager.cpp b/Engine/Plugins/Experimental/DigitalMultiplex/DMXProtocol/Source/DMXProtocol/Private/Managers/DMXProtocolPortManager.cpp
new file mode 100644
--- /dev/null
+++ b/Engine/Plugins/Experimental/DigitalMultiplex/DMXProtocol/Source/DMXProtocol/Private/Managers/DMXProtocolPortManager.cpp
@@ -0,0 +1,104 @@
+// Copyright 1998-2019 Epic Games, Inc. All Rights Reserved.
+
+#include "Managers/DMXProtocolPortManager.h"
+#include "Interfaces/IDMXProtocolPort.h"
+#include "Interfaces/IDMXProtocol.h"
+#include "Interfaces/IDMXProtocolDevice.h"
+
+void FDMXProtocolPortManager::AddInputPort(TSharedPtr<IDMXProtocolDevice> InDevice, TSharedPtr<IDMXProtocolPort> InInputPort)
+{
+	checkf(InInputPort.IsValid(), TEXT("IDMXProtocolInputPortPtr not valid. Manager should own only valid pointer."));
+	checkf(InDevice.IsValid(), TEXT("InDevice not valid"));
+
+	if (IInputPortsDeviceMap.Contains(InDevice.Get()))
+	{
+		IPortsMap& InputPortsMap = IInputPortsDeviceMap[InDevice.Get()];
+		InputPortsMap.Add(InInputPort->GetPortID(), InInputPort);
+	}
+	else
+	{
+		IPortsMap InputPortsMap;
+		InputPortsMap.Add(InInputPort->GetPortID(), InInputPort);
+		IInputPortsDeviceMap.Add(InDevice.Get(), InputPortsMap);
+	}
+}
+
+void FDMXProtocolPortManager::AddOutputPort(TSharedPtr<IDMXProtocolDevice> InDevice, TSharedPtr<IDMXProtocolPort> InOutputPort)
+{
+	checkf(InOutputPort.IsValid(), TEXT("IDMXProtocolOutputPortPtr not valid"));
+	checkf(InDevice.IsValid(), TEXT("InDevice not valid"));
+
+	if (IOutputPortsDeviceMap.Contains(InDevice.Get()))
+	{
+		IPortsMap& OutputPortsMap = IOutputPortsDeviceMap[InDevice.Get()];
+		OutputPortsMap.Add(InOutputPort->GetPortID(), InOutputPort);
+	}
+	else
+	{
+		IPortsMap OutputPortsMap;
+		OutputPortsMap.Add(InOutputPort->GetPortID(), InOutputPort);
+		IOutputPortsDeviceMap.Add(InDevice.Get(), OutputPortsMap);
+	}
+}
+
+void FDMXProtocolPortManager::RemoveInputPort(TSharedPtr<IDMXProtocolDevice> InDevice, TSharedPtr<IDMXProtocolPort> InInputPort)
+{
+	if (IInputPortsDeviceMap.Contains(InDevice.Get()))
+	{
+		IPortsMap& InputPortsMap = IInputPortsDeviceMap[InDevice.Get()];
+		InputPortsMap.Remove(InInputPort->GetPortID());
+	}
+}
+
+void FDMXProtocolPortManager::RemoveOutputPort(TSharedPtr<IDMXProtocolDevice> InDevice, TSharedPtr<IDMXProtocolPort> InOutputPort)
+{
+	if (IOutputPortsDeviceMap.Contains(InDevice.Get()))
+	{
+		IPortsMap& OutputPortsMap = IOutputPortsDeviceMap[InDevice.Get()];
+		OutputPortsMap.Remove(InOutputPort->GetPortID());
+	}
+}
+
+IPortsMap* FDMXProtocolPortManager::GetInputPortMapByDevice(IDMXProtocolDevice* InDevice)
+{
+	return IOutputPortsDeviceMap.Find(InDevice);
+}
+
+IPortsMap* FDMXProtocolPortManager::GetOutputPortMapByDevice(IDMXProtocolDevice* InDevice)
+{
+	return IOutputPortsDeviceMap.Find(InDevice);
+}
+
+void FDMXProtocolPortManager::RemoveAll()
+{
+	IInputPortsDeviceMap.Empty();
+	IOutputPortsDeviceMap.Empty();
+}
+
+IDMXProtocolPort * FDMXProtocolPortManager::GetPortByDeviceAndID(const TSharedPtr<IDMXProtocolDevice>& InDevice, uint8 InPortID, EDMXPortDirection Direction)
+{
+	IPortsMap* PortsMap = nullptr;
+	switch (Direction)
+	{
+	case DMX_PORT_OUTPUT:
+		PortsMap = GetInputPortMapByDevice(InDevice.Get());
+		break;
+	case DMX_PORT_INPUT:
+		PortsMap = GetOutputPortMapByDevice(InDevice.Get());
+		break;
+	}
+
+	if (PortsMap != nullptr)
+	{
+		TSharedPtr<IDMXProtocolPort>* Port = PortsMap->Find(InPortID);
+		if (Port != nullptr && Port->IsValid())
+		{
+			if (Port->Get()->GetCachedDevice().HasSameObject(InDevice.Get()))
+			{
+				return Port->Get();
+			}
+		}
+	}
+
+	return nullptr;
+}
\ No newline at end of file
diff --git a/Engine/Plugins/Experimental/DigitalMultiplex/DMXProtocol/Source/DMXProtocol/Private/Managers/DMXProtocolUniverseManager.cpp b/Engine/Plugins/Experimental/DigitalMultiplex/DMXProtocol/Source/DMXProtocol/Private/Managers/DMXProtocolUniverseManager.cpp
new file mode 100644
--- /dev/null
+++ b/Engine/Plugins/Experimental/DigitalMultiplex/DMXProtocol/Source/DMXProtocol/Private/Managers/DMXProtocolUniverseManager.cpp
@@ -0,0 +1,54 @@
+// Copyright 1998-2019 Epic Games, Inc. All Rights Reserved.
+
+#include "Managers/DMXProtocolUniverseManager.h"
+#include "Interfaces/IDMXProtocolUniverse.h"
+#include "Interfaces/IDMXProtocol.h"
+
+void FDMXProtocolUniverseManager::AddUniverse(TSharedPtr<IDMXProtocolPort> InPort, TSharedPtr<IDMXProtocolUniverse> InUniverse)
+{
+	checkf(InPort.IsValid(), TEXT("IDMXProtocolPortWeakPtr not valid"));
+	checkf(InUniverse.IsValid(), TEXT("IDMXProtocolUniversePtr not valid. Manager should own only valid pointer."));
+	UniversesMap.Add(InPort.Get(), InUniverse);
+	UniversesIDMap.Add(InUniverse->GetUniverseID(), InPort.Get());
+}
+
+void FDMXProtocolUniverseManager::RemoveUniverse(TSharedPtr<IDMXProtocolPort> InPort, TSharedPtr<IDMXProtocolUniverse> InUniverse)
+{
+	UniversesMap.Remove(InPort.Get());
+	UniversesIDMap.Remove(InUniverse->GetUniverseID());
+}
+
+
+void FDMXProtocolUniverseManager::RemoveAll()
+{
+	UniversesMap.Empty();
+	UniversesIDMap.Empty();
+}
+
+IDMXProtocolUniverse* FDMXProtocolUniverseManager::GetUniverseByPort(IDMXProtocolPort* InPort)
+{
+	TSharedPtr<IDMXProtocolUniverse>* Port = UniversesMap.Find(InPort);
+
+	if (Port != nullptr)
+	{
+		return Port->Get();
+	}
+
+	return nullptr;
+}
+
+IDMXProtocolUniverse* FDMXProtocolUniverseManager::GetUniverseById(uint16 UniverseID)
+{
+	IDMXProtocolPort** PortAddr = UniversesIDMap.Find(UniverseID);
+
+	if (PortAddr != nullptr)
+	{
+		TSharedPtr<IDMXProtocolUniverse>* Port = UniversesMap.Find(*PortAddr);
+
+		if (Port != nullptr)
+		{
+			return Port->Get();
+		}
+	}
+	return nullptr;
+}
diff --git a/Engine/Plugins/Experimental/DigitalMultiplex/DMXProtocol/Source/DMXProtocol/Private/Tests/DMXProtocolTests.cpp b/Engine/Plugins/Experimental/DigitalMultiplex/DMXProtocol/Source/DMXProtocol/Private/Tests/DMXProtocolTests.cpp
new file mode 100644
--- /dev/null
+++ b/Engine/Plugins/Experimental/DigitalMultiplex/DMXProtocol/Source/DMXProtocol/Private/Tests/DMXProtocolTests.cpp
@@ -0,0 +1,159 @@
+// Copyright 1998-2019 Epic Games, Inc. All Rights Reserved.
+
+#include "CoreMinimal.h"
+#include "Misc/AutomationTest.h"
+#include "IDMXProtocolFactory.h"
+#include "DMXProtocolCommon.h"
+#include "Json.h"
+
+#include "DMXProtocol.h"
+
+
+class FDMXProtocolTest
+	: public FDMXProtocol
+{
+public:
+	//~ Begin IDMXProtocolBase implementation
+	virtual bool Init() override { return true; }
+	virtual bool Shutdown() override { return true; }
+	virtual bool Tick(float DeltaTime) override { return true; }
+	//~ End IDMXProtocolBase implementation
+
+	//~ Begin IDMXProtocol implementation
+	virtual TSharedPtr<IDMXProtocolSender> GetSenderInterface() const override { return nullptr; }
+	virtual FName GetProtocolName() const override { return ProtocolName; }
+	virtual void Reload() override {}
+	virtual bool SendDMX(uint16 UniverseID, uint8 PortID, const TSharedPtr<FDMXBuffer>& DMXBuffer) const override { return true; }
+	virtual bool IsEnabled() const override { return true; }
+	//~ End IDMXProtocol implementation
+
+	//~ Begin IDMXProtocolRDM implementation
+	virtual void SendRDMCommand(const TSharedPtr<FJsonObject>& CMD) override {}
+	virtual void RDMDiscovery(const TSharedPtr<FJsonObject>& CMD) override {}
+	//~ End IDMXProtocol implementation
+
+	//~ Only the factory makes instances
+	FDMXProtocolTest() = delete;
+	explicit FDMXProtocolTest(FName InProtocolName, FJsonObject& InSettings)
+		: FDMXProtocol(InProtocolName, InSettings)
+		, ProtocolName(InProtocolName)
+	{}
+
+private:
+	FName ProtocolName;
+};
+
+
+/**
+ */
+class FDMXProtocolFactoryTestFactory : public IDMXProtocolFactory
+{
+public:
+	virtual TSharedPtr<IDMXProtocol> CreateProtocol(const FName& ProtocolName) override
+	{
+		FJsonObject ProtocolSettings;
+		TSharedPtr<IDMXProtocol> ProtocolArtNetPtr = MakeShared<FDMXProtocolTest>(ProtocolName, ProtocolSettings);
+		if (ProtocolArtNetPtr->IsEnabled())
+		{
+			if (!ProtocolArtNetPtr->Init())
+			{
+				UE_LOG_DMXPROTOCOL(Warning, TEXT("TEST Protocol failed to initialize!"));
+				ProtocolArtNetPtr->Shutdown();
+				ProtocolArtNetPtr = nullptr;
+			}
+		}
+		else
+		{
+			UE_LOG_DMXPROTOCOL(Warning, TEXT("TEST Protocol disabled!"));
+			ProtocolArtNetPtr->Shutdown();
+			ProtocolArtNetPtr = nullptr;
+		}
+
+		return ProtocolArtNetPtr;
+	}
+};
+
+namespace DMXProtocolTestHelper
+{
+	static const FName NAME_ArtnetTest = TEXT("ARTNET_TEST");
+	static const FName NAME_SACNTest = TEXT("SACN_TEST");
+
+	void GetDMXProtocolNamesForTesting(TArray<FString>& OutBeautifiedNames, TArray<FString>& OutTestCommands, const FString& PostTestName = TEXT(""))
+	{
+		TArray< FName > DMXProtocolList;
+		DMXProtocolList.AddUnique(NAME_ArtnetTest);
+		DMXProtocolList.AddUnique(NAME_SACNTest);
+
+		for (FName& DMXProtocolName : DMXProtocolList)
+		{
+			FString PostName = FString::Printf(TEXT(".%s"), *PostTestName);
+			FString PrettyName = FString::Printf(TEXT("%s%s"),
+				*DMXProtocolName.ToString(),
+				PostTestName.IsEmpty() ? TEXT("") : *PostName);
+			OutBeautifiedNames.Add(PrettyName);
+			OutTestCommands.Add(DMXProtocolName.ToString());
+		}
+	}
+}
+
+DEFINE_LATENT_AUTOMATION_COMMAND_ONE_PARAMETER(FDMXProtocolFactoryTestCommand, FName, ProtocolName);
+
+bool FDMXProtocolFactoryTestCommand::Update()
+{
+	//UE_LOG(LogTemp, Warning, TEXT("ProtocolName '%s'"), *ProtocolName.ToString());
+
+	return true;
+}
+
+IMPLEMENT_COMPLEX_AUTOMATION_TEST(FDMXProtocolFactoryTest, "VirtualProduction.DMX.Protocol.Factory", (EAutomationTestFlags::ApplicationContextMask | EAutomationTestFlags::EngineFilter))
+
+void FDMXProtocolFactoryTest::GetTests(TArray<FString>& OutBeautifiedNames, TArray<FString>& OutTestCommands) const
+{
+	DMXProtocolTestHelper::GetDMXProtocolNamesForTesting(
+		OutBeautifiedNames,
+		OutTestCommands,
+		TEXT("Functional test of the Protocol factory")
+	);
+}
+
+bool FDMXProtocolFactoryTest::RunTest(const FString& Parameters)
+{
+	// parameter is the provider we want to use
+	FName ProtocolName = FName(*Parameters);
+	ADD_LATENT_AUTOMATION_COMMAND(FDMXProtocolFactoryTestCommand(ProtocolName));
+
+	FDMXProtocolModule& DMXProtocolModule = FModuleManager::GetModuleChecked<FDMXProtocolModule>("DMXProtocol");
+
+	// Store the protocol pointer
+	IDMXProtocol* CachedProtocol = nullptr;
+
+	// Try to register 3 times
+	TArray<TUniquePtr<IDMXProtocolFactory>> Factories;
+	Factories.Add(MakeUnique<FDMXProtocolFactoryTestFactory>());
+	Factories.Add(MakeUnique<FDMXProtocolFactoryTestFactory>());
+	Factories.Add(MakeUnique<FDMXProtocolFactoryTestFactory>());
+	for (int32 Index = 0; Index < Factories.Num(); Index++)
+	{
+		// Create and register our singleton factory with the main online subsystem for easy access
+		DMXProtocolModule.RegisterProtocol(ProtocolName, Factories[Index].Get());
+		if (CachedProtocol == nullptr)
+		{
+			CachedProtocol = IDMXProtocol::Get(ProtocolName);
+		}
+
+		TestNotNull(TEXT("Protocol should exists"), IDMXProtocol::Get(ProtocolName));
+		TestEqual(TEXT("Should return same protocol instance"), CachedProtocol, IDMXProtocol::Get(ProtocolName));
+	}
+	Factories.Empty();
+
+	// Protocol removal test
+	{
+		TUniquePtr<IDMXProtocolFactory> Factory = MakeUnique<FDMXProtocolFactoryTestFactory>();
+		DMXProtocolModule.RegisterProtocol(ProtocolName, Factory.Get());
+		DMXProtocolModule.UnregisterProtocol(ProtocolName);
+		TestNull(TEXT("Protocol should not exists"), IDMXProtocol::Get(ProtocolName));
+	}
+
+
+	return true;
+}
\ No newline at end of file
diff --git a/Engine/Plugins/Experimental/DigitalMultiplex/DMXProtocol/Source/DMXProtocol/Public/DMXProtocol.h b/Engine/Plugins/Experimental/DigitalMultiplex/DMXProtocol/Source/DMXProtocol/Public/DMXProtocol.h
new file mode 100644
--- /dev/null
+++ b/Engine/Plugins/Experimental/DigitalMultiplex/DMXProtocol/Source/DMXProtocol/Public/DMXProtocol.h
@@ -0,0 +1,38 @@
+// Copyright 1998-2019 Epic Games, Inc. All Rights Reserved.
+
+#pragma once
+
+#include "DMXProtocolCommon.h"
+#include "Interfaces/IDMXProtocol.h"
+
+class DMXPROTOCOL_API FDMXProtocol 
+	: public IDMXProtocol
+{
+public:
+
+	//~ Begin IDMXProtocol implementation
+	virtual TSharedPtr<FDMXProtocolDeviceManager> GetDeviceManager() const override;
+	
+	virtual TSharedPtr<FDMXProtocolInterfaceManager> GetInterfaceManager() const override;
+
+	virtual TSharedPtr<FDMXProtocolPortManager> GetPortManager() const override;
+	
+	virtual TSharedPtr<FDMXProtocolUniverseManager> GetUniverseManager() const override;
+
+	virtual TSharedPtr<FJsonObject> GetSettings() const override;
+
+	virtual void SetDMXFragment(uint16 UniverseID, const IDMXFragmentMap& DMXFragment, bool bShouldSend = true) override;
+	//~ End IDMXProtocol implementation
+
+protected:
+	FDMXProtocol() = delete;
+	FDMXProtocol(const FName& InProtocolName, FJsonObject& InSettings);
+	~FDMXProtocol();
+
+protected:
+	TSharedPtr<FDMXProtocolDeviceManager> DeviceManager;
+	TSharedPtr<FDMXProtocolInterfaceManager> InterfaceManager;
+	TSharedPtr<FDMXProtocolPortManager> PortManager;
+	TSharedPtr<FDMXProtocolUniverseManager> UniverseManager;
+	TSharedPtr<FJsonObject> Settings;
+};
\ No newline at end of file
diff --git a/Engine/Plugins/Experimental/DigitalMultiplex/DMXProtocol/Source/DMXProtocol/Public/DMXProtocolCommon.h b/Engine/Plugins/Experimental/DigitalMultiplex/DMXProtocol/Source/DMXProtocol/Public/DMXProtocolCommon.h
new file mode 100644
--- /dev/null
+++ b/Engine/Plugins/Experimental/DigitalMultiplex/DMXProtocol/Source/DMXProtocol/Public/DMXProtocolCommon.h
@@ -0,0 +1,44 @@
+#pragma once
+
+#include "CoreMinimal.h"
+#include "DMXProtocolConstants.h"
+
+class FBufferArchive;
+class FJsonObject;
+class FArrayReader;
+class IDMXProtocol;
+class IDMXProtocolDevice;
+class IDMXProtocolFactory;
+class IDMXProtocolInterface;
+class IDMXProtocolInterfaceUSB;
+class IDMXProtocolInterfaceEthernet;
+class IDMXProtocolPort;
+class IDMXProtocolRDM;
+class IDMXProtocolUniverse;
+class IDMXProtocolSender;
+class IDMXProtocolReceiver;
+
+class FDMXProtocolDeviceManager;
+class FDMXProtocolInterfaceManager;
+class FDMXProtocolPortManager;
+class FDMXProtocolUniverseManager;
+
+class FDMXProtocolPackager;
+
+struct FDMXBuffer;
+struct FDMXPacket;
+struct IDMXProtocolPacket;
+
+using FDataPtr = TSharedPtr<TArray<uint8>>;
+using IDevicesMap = TMap<IDMXProtocolInterface*, TSharedPtr<IDMXProtocolDevice>>;
+using IInterfacesMap = TMap<uint32, TSharedPtr<IDMXProtocolInterface>>;
+
+using IPortsMap = TMap<uint8, TSharedPtr<IDMXProtocolPort>>;
+using IPortsDeviceMap = TMap<IDMXProtocolDevice*, IPortsMap>;
+using IUniversesMap = TMap<IDMXProtocolPort*, TSharedPtr<IDMXProtocolUniverse>>;
+using IUniversesIDMap = TMap<uint16, IDMXProtocolPort*>;
+
+using IDMXFragmentMap = TMap<uint32, uint8>;
+
+using FArrayReaderPtr = TSharedPtr<FArrayReader, ESPMode::ThreadSafe>;
+using FDMXPacketPtr = TSharedPtr<FDMXPacket, ESPMode::ThreadSafe>;
\ No newline at end of file
diff --git a/Engine/Plugins/Experimental/DigitalMultiplex/DMXProtocol/Source/DMXProtocol/Public/DMXProtocolConstants.h b/Engine/Plugins/Experimental/DigitalMultiplex/DMXProtocol/Source/DMXProtocol/Public/DMXProtocolConstants.h
new file mode 100644
--- /dev/null
+++ b/Engine/Plugins/Experimental/DigitalMultiplex/DMXProtocol/Source/DMXProtocol/Public/DMXProtocolConstants.h
@@ -0,0 +1,26 @@
+// Copyright 1998-2019 Epic Games, Inc. All Rights Reserved.
+
+#pragma once
+
+#include "CoreMinimal.h"
+
+enum 
+{
+  DMX_UNIVERSE_SIZE = 512,
+  DMX_MAX_CHANNEL_VALUE = 255,
+  RDM_UID_WIDTH = 6
+};
+
+enum EDMXPortCapability
+{
+	DMX_PORT_CAPABILITY_NONE,
+	DMX_PORT_CAPABILITY_STATIC,
+	DMX_PORT_CAPABILITY_FULL,
+};
+
+enum EDMXPortDirection
+{
+	DMX_PORT_UNKNOWN,
+	DMX_PORT_OUTPUT,
+	DMX_PORT_INPUT
+};
\ No newline at end of file
diff --git a/Engine/Plugins/Experimental/DigitalMultiplex/DMXProtocol/Source/DMXProtocol/Public/DMXProtocolMacros.h b/Engine/Plugins/Experimental/DigitalMultiplex/DMXProtocol/Source/DMXProtocol/Public/DMXProtocolMacros.h
new file mode 100644
--- /dev/null
+++ b/Engine/Plugins/Experimental/DigitalMultiplex/DMXProtocol/Source/DMXProtocol/Public/DMXProtocolMacros.h
@@ -0,0 +1,12 @@
+// Copyright 1998-2019 Epic Games, Inc. All Rights Reserved.
+
+#pragma once
+
+#include "CoreMinimal.h"
+
+#define REGISTER_DMX_ARCHIVE(ProtocolDMXPacket) \
+	FArchive & operator<<(FArchive & Ar, ProtocolDMXPacket&Packet) \
+	{ \
+		Packet.Serialize(Ar); \
+		return Ar; \
+	}
\ No newline at end of file
diff --git a/Engine/Plugins/Experimental/DigitalMultiplex/DMXProtocol/Source/DMXProtocol/Public/DMXProtocolModule.h b/Engine/Plugins/Experimental/DigitalMultiplex/DMXProtocol/Source/DMXProtocol/Public/DMXProtocolModule.h
new file mode 100644
--- /dev/null
+++ b/Engine/Plugins/Experimental/DigitalMultiplex/DMXProtocol/Source/DMXProtocol/Public/DMXProtocolModule.h
@@ -0,0 +1,44 @@
+// Copyright 1998-2019 Epic Games, Inc. All Rights Reserved.
+
+#pragma once
+
+#include "DMXProtocolCommon.h"
+#include "Modules/ModuleInterface.h"
+
+/**
+ * 
+ */
+class DMXPROTOCOL_API FDMXProtocolModule 
+	: public IModuleInterface
+{
+public:
+	void RegisterProtocol(const FName& FactoryName, IDMXProtocolFactory* Factory);
+
+	void UnregisterProtocol(const FName& FactoryName);
+
+	/**
+	 * If protocol exists return the pointer otherwise it create a new protocol first and then return the pointer.
+	 * @param  InProtocolName Name of the requested protocol
+	 * @return Return the pointer to protocol.
+	 */
+	virtual IDMXProtocol* GetProtocol(const FName InProtocolName = NAME_None);
+
+	//~ Begin IModuleInterface implementation
+	virtual void ShutdownModule() override;
+	//~ End IModuleInterface implementation
+
+	/** Get the instance of this module. */
+	static FDMXProtocolModule& Get();
+
+private:
+	void ShutdownDMXProtocol(const FName& ProtocolName);
+	void ShutdownAllDMXProtocols();
+
+public:
+	static const TCHAR* BaseModuleName;
+
+private:
+	TMap<FName, IDMXProtocolFactory*> DMXFactories;
+	TMap<FName, TSharedPtr<IDMXProtocol>> DMXProtocols;
+	TMap<FName, bool> DMXProtocolFailureNotes;
+};
diff --git a/Engine/Plugins/Experimental/DigitalMultiplex/DMXProtocol/Source/DMXProtocol/Public/DMXProtocolPackager.h b/Engine/Plugins/Experimental/DigitalMultiplex/DMXProtocol/Source/DMXProtocol/Public/DMXProtocolPackager.h
new file mode 100644
--- /dev/null
+++ b/Engine/Plugins/Experimental/DigitalMultiplex/DMXProtocol/Source/DMXProtocol/Public/DMXProtocolPackager.h
@@ -0,0 +1,17 @@
+// Copyright 1998-2019 Epic Games, Inc. All Rights Reserved.
+
+#pragma once
+
+#include "DMXProtocolCommon.h"
+
+class DMXPROTOCOL_API FDMXProtocolPackager
+{
+public:
+	bool AddToPackage(IDMXProtocolPacket* Packet);
+	uint32 GetSize() const;
+	const TArray<uint8>& GetBuffer() const;
+
+private:
+	TArray<uint8> Buffer;
+};
+
diff --git a/Engine/Plugins/Experimental/DigitalMultiplex/DMXProtocol/Source/DMXProtocol/Public/DMXProtocolSettings.h b/Engine/Plugins/Experimental/DigitalMultiplex/DMXProtocol/Source/DMXProtocol/Public/DMXProtocolSettings.h
new file mode 100644
--- /dev/null
+++ b/Engine/Plugins/Experimental/DigitalMultiplex/DMXProtocol/Source/DMXProtocol/Public/DMXProtocolSettings.h
@@ -0,0 +1,28 @@
+// Copyright 1998-2019 Epic Games, Inc. All Rights Reserved.
+
+#pragma once
+
+#include "CoreMinimal.h"
+#include "UObject/Object.h"
+#include "DMXProtocolSettings.generated.h"
+
+/**
+ * User defined protocol settings that apply to a whole protocol module
+ */
+UCLASS(config = Engine)
+class DMXPROTOCOL_API UDMXProtocolSettings : public UObject
+{
+public:
+	GENERATED_BODY()
+
+public:
+	UDMXProtocolSettings() :
+		InterfaceIPAddress(TEXT("0.0.0.0"))
+	{}
+
+
+public:
+	/** Manual Interface IP Address */
+	UPROPERTY(GlobalConfig, EditAnywhere, Category = "DMX", Meta = (DisplayName = "Interface IP address"))
+	FString InterfaceIPAddress;
+};
diff --git a/Engine/Plugins/Experimental/DigitalMultiplex/DMXProtocol/Source/DMXProtocol/Public/DMXProtocolTypes.h b/Engine/Plugins/Experimental/DigitalMultiplex/DMXProtocol/Source/DMXProtocol/Public/DMXProtocolTypes.h
new file mode 100644
--- /dev/null
+++ b/Engine/Plugins/Experimental/DigitalMultiplex/DMXProtocol/Source/DMXProtocol/Public/DMXProtocolTypes.h
@@ -0,0 +1,69 @@
+// Copyright 1998-2019 Epic Games, Inc. All Rights Reserved.
+
+#pragma once
+
+#include "IDMXProtocol.h"
+#include "DMXProtocolCommon.h"
+#include "DMXProtocolConstants.h"
+#include "Json.h"
+
+struct DMXPROTOCOL_API FDMXBuffer
+{
+public:
+	FDMXBuffer()
+	{
+		DMXData.AddZeroed(DMX_UNIVERSE_SIZE);
+	}
+
+	bool SetDMXFragment(const IDMXFragmentMap & DMXFragment);
+	TArray<uint8>& GetDMXData() { return DMXData; }
+
+private:
+	TArray<uint8> DMXData;
+};
+
+struct DMXPROTOCOL_API FDMXPacket
+{
+	FDMXPacket(FJsonObject& InSettings, const TArray<uint8>& InData)
+	{
+		Settings = InSettings;
+		Data = InData;
+	}
+
+	FDMXPacket(const TArray<uint8>& InData)
+	{
+		Data = InData;
+	}
+
+	FJsonObject Settings;
+	TArray<uint8> Data;
+};
+
+struct DMXPROTOCOL_API FRDMUID
+{
+	FRDMUID()
+	{
+	}
+
+	FRDMUID(uint8 InBuffer[RDM_UID_WIDTH])
+	{
+		FMemory::Memcpy(Buffer, InBuffer, RDM_UID_WIDTH);
+	}
+
+	FRDMUID(const TArray<uint8>& InBuffer)
+	{
+		// Only copy if we have the requested amount of data
+		if (InBuffer.Num() == RDM_UID_WIDTH)
+		{
+			FMemory::Memmove(Buffer, InBuffer.GetData(), RDM_UID_WIDTH);
+		}
+		else
+		{
+			UE_LOG_DMXPROTOCOL(Warning, TEXT("Size of the TArray buffer is wrong"));
+		}
+	}
+
+	uint8 Buffer[RDM_UID_WIDTH] = { 0 };
+};
+
+
diff --git a/Engine/Plugins/Experimental/DigitalMultiplex/DMXProtocol/Source/DMXProtocol/Public/Interfaces/IDMXProtocol.h b/Engine/Plugins/Experimental/DigitalMultiplex/DMXProtocol/Source/DMXProtocol/Public/Interfaces/IDMXProtocol.h
new file mode 100644
--- /dev/null
+++ b/Engine/Plugins/Experimental/DigitalMultiplex/DMXProtocol/Source/DMXProtocol/Public/Interfaces/IDMXProtocol.h
@@ -0,0 +1,149 @@
+// Copyright 1998-2019 Epic Games, Inc. All Rights Reserved.
+
+#pragma once
+
+#include "DMXProtocolCommon.h"
+#include "Stats/Stats.h"
+#include "Modules/ModuleManager.h"
+#include "DMXProtocolModule.h"
+
+#include "Interfaces/IDMXProtocolBase.h"
+#include "Interfaces/IDMXProtocolRDM.h"
+#include "Interfaces/IDMXProtocolTransport.h"
+
+#include "Managers/DMXProtocolDeviceManager.h"
+#include "Managers/DMXProtocolInterfaceManager.h"
+#include "Managers/DMXProtocolPortManager.h"
+#include "Managers/DMXProtocolUniverseManager.h"
+
+DMXPROTOCOL_API DECLARE_LOG_CATEGORY_EXTERN(LogDMXProtocol, Log, All);
+
+DECLARE_STATS_GROUP(TEXT("DMXProtocol"), STATGROUP_DMXProtocol, STATCAT_Advanced);
+
+#ifndef DMXPROTOCOL_LOG_PREFIX
+#define DMXPROTOCOL_LOG_PREFIX TEXT("DMX: ")
+#endif
+
+#define UE_LOG_DMXPROTOCOL(Verbosity, Format, ...) \
+{ \
+	UE_LOG(LogDMXProtocol, Verbosity, TEXT("%s%s"), DMXPROTOCOL_LOG_PREFIX, *FString::Printf(Format, ##__VA_ARGS__)); \
+}
+
+#define UE_CLOG_DMXPROTOCOL(Conditional, Verbosity, Format, ...) \
+{ \
+	UE_CLOG(Conditional, LogDMXProtocol, Verbosity, TEXT("%s%s"), DMXPROTOCOL_LOG_PREFIX, *FString::Printf(Format, ##__VA_ARGS__)); \
+}
+
+/**
+ * Generic protocol interface, it should be inherited by all protocol implementations.
+ */
+class DMXPROTOCOL_API IDMXProtocol 
+	: public IDMXProtocolBase
+	, public IDMXProtocolRDM
+{
+public:
+	virtual ~IDMXProtocol() {}
+
+	/**
+	 * If protocol exists return the pointer otherwise it create a new protocol first and then return the pointer.
+	 * @param  InProtocolName Name of the requested protocol
+	 * @return Return the pointer to protocol
+	 */
+	template <class TProtocol>
+	static TProtocol* Get(const FName& ProtocolName = NAME_None)
+	{
+		static const FName DMXProtocolModuleName = TEXT("DMXProtocol");
+		FDMXProtocolModule& DMXProtocolModule = FModuleManager::GetModuleChecked<FDMXProtocolModule>(DMXProtocolModuleName);
+		return static_cast<TProtocol*>(DMXProtocolModule.GetProtocol(ProtocolName));
+	}
+
+	/**
+	 * If protocol exists return the pointer otherwise it create a new protocol first and then return the pointer.
+	 * @param  InProtocolName Name of the requested protocol
+	 * @return Return the pointer to protocol
+	 */
+	static IDMXProtocol* Get(const FName& ProtocolName = NAME_None)
+	{
+		static const FName DMXProtocolModuleName = TEXT("DMXProtocol");
+		FDMXProtocolModule& DMXProtocolModule = FModuleManager::GetModuleChecked<FDMXProtocolModule>(DMXProtocolModuleName);
+		return DMXProtocolModule.GetProtocol(ProtocolName);
+	}
+
+	/**
+	 * Get the Protocol Name
+	 * @return Return FName of the protocol
+	 */
+	virtual FName GetProtocolName() const = 0;
+	
+	/**
+	 * Get the Protocol Sender Interface
+	 * Sender interface holds the functionality to queue and physically send the DMX buffer
+	 * @return Return the pointer to SenderInterface
+	 */
+	virtual TSharedPtr<IDMXProtocolSender> GetSenderInterface() const = 0;
+
+	/**
+	 * Get the device manager for a protocol instance
+	 * The device manager is responsible for hold and searches the physical devices
+	 * Physical devices are nodes, controllers, consoles
+	 * @return Return the pointer to DeviceManager
+	 */
+	virtual TSharedPtr<FDMXProtocolDeviceManager> GetDeviceManager() const = 0;
+	
+	/**
+	 * Get the interface manager for a protocol instance
+	 * The interface is the way how protocol interacts with Unreal Engine
+	 * This could be done with ethernet, serial or fake communication
+	 * @return Return the pointer to InterfaceManager
+	 */
+	virtual TSharedPtr<FDMXProtocolInterfaceManager> GetInterfaceManager() const = 0;
+
+	/**
+	 * Get the port manager for a protocol instance
+	 * Port managers are responsible for hold the configuration instances of physical ports
+	 */
+	virtual TSharedPtr<FDMXProtocolPortManager> GetPortManager() const = 0;
+	
+	/**
+	 * Get the universe manager for a protocol instance
+	 * Universe managers hold the buffers of DMX for particular universes
+	 * @return Return the pointer to UniverseManager
+	 */
+	virtual TSharedPtr<FDMXProtocolUniverseManager> GetUniverseManager() const = 0;
+
+	/**
+	 * Get the protocol settings
+	 * @return Return the pointer to Protocol Settings
+	 */
+	virtual TSharedPtr<FJsonObject> GetSettings() const = 0;
+
+	/**
+	 * Reloads the protocol instance
+	 */
+	virtual void Reload() = 0;
+
+	/**
+	 * Whether protocol enabled
+	 * @return Return true if the protocol is enabled
+	 */
+	virtual bool IsEnabled() const = 0;
+
+	/**
+	 * Sets the DMX fragment for a particular universe
+	 * @param  UniverseID ID of universe to send
+	 * @param  DMXFragment Map of DMX channel  and values
+	 * @param  bShouldSend whether it should be sent or just store in the universe buffer
+	 */
+	virtual void SetDMXFragment(uint16 UniverseID, const IDMXFragmentMap& DMXFragment, bool bShouldSend = true) = 0;
+
+	/**
+	 * Send DMX to physical device port
+	 * @param  UniverseID ID of universe to send
+	 * @param  PortID Physical port ID
+	 * @param  DMXBuffer DMX buffer to sent
+	 * @return Return true if DMX buffer included in the send queue
+	 */
+	virtual bool SendDMX(uint16 UniverseID, uint8 PortID, const TSharedPtr<FDMXBuffer>& DMXBuffer) const = 0;
+};
+
+
diff --git a/Engine/Plugins/Experimental/DigitalMultiplex/DMXProtocol/Source/DMXProtocol/Public/Interfaces/IDMXProtocolBase.h b/Engine/Plugins/Experimental/DigitalMultiplex/DMXProtocol/Source/DMXProtocol/Public/Interfaces/IDMXProtocolBase.h
new file mode 100644
--- /dev/null
+++ b/Engine/Plugins/Experimental/DigitalMultiplex/DMXProtocol/Source/DMXProtocol/Public/Interfaces/IDMXProtocolBase.h
@@ -0,0 +1,16 @@
+// Copyright 1998-2019 Epic Games, Inc. All Rights Reserved.
+
+#pragma once
+
+#include "CoreMinimal.h"
+
+class DMXPROTOCOL_API IDMXProtocolBase
+{
+public:
+	virtual bool Init() = 0;
+
+	virtual bool Shutdown() = 0;
+
+	virtual bool Tick(float DeltaTime) = 0;
+};
+
diff --git a/Engine/Plugins/Experimental/DigitalMultiplex/DMXProtocol/Source/DMXProtocol/Public/Interfaces/IDMXProtocolDevice.h b/Engine/Plugins/Experimental/DigitalMultiplex/DMXProtocol/Source/DMXProtocol/Public/Interfaces/IDMXProtocolDevice.h
new file mode 100644
--- /dev/null
+++ b/Engine/Plugins/Experimental/DigitalMultiplex/DMXProtocol/Source/DMXProtocol/Public/Interfaces/IDMXProtocolDevice.h
@@ -0,0 +1,21 @@
+// Copyright 1998-2019 Epic Games, Inc. All Rights Reserved.
+
+#pragma once
+
+#include "DMXProtocolCommon.h"
+
+/**
+ *	IDMXProtocolDevice - Physical Device. Holds the configuration for ports, ip, mac addresses, etc.
+ */
+class DMXPROTOCOL_API IDMXProtocolDevice
+{
+public:
+	virtual ~IDMXProtocolDevice() {}
+
+	virtual TSharedPtr<FJsonObject> GetSettings() const = 0;
+	virtual TWeakPtr<IDMXProtocolInterface> GetCachedProtocolInterface() const = 0;
+	virtual IDMXProtocol* GetProtocol() const = 0;
+	virtual uint32 GetDeviceID() const = 0;
+
+	virtual bool AllowLooping() const = 0;
+};
diff --git a/Engine/Plugins/Experimental/DigitalMultiplex/DMXProtocol/Source/DMXProtocol/Public/Interfaces/IDMXProtocolFactory.h b/Engine/Plugins/Experimental/DigitalMultiplex/DMXProtocol/Source/DMXProtocol/Public/Interfaces/IDMXProtocolFactory.h
new file mode 100644
--- /dev/null
+++ b/Engine/Plugins/Experimental/DigitalMultiplex/DMXProtocol/Source/DMXProtocol/Public/Interfaces/IDMXProtocolFactory.h
@@ -0,0 +1,17 @@
+// Copyright 1998-2019 Epic Games, Inc. All Rights Reserved.
+
+#pragma once
+
+#include "DMXProtocolCommon.h"
+
+class DMXPROTOCOL_API IDMXProtocolFactory
+{
+public:
+	virtual ~IDMXProtocolFactory() {}
+
+	/**
+	 * It creates only one instance of Protocol. There is the condition to check in DMXProtocolModule
+	 * @return Return the pointer to protocol
+	 */
+	virtual TSharedPtr<IDMXProtocol> CreateProtocol(const FName& ProtocolName) = 0;
+};
diff --git a/Engine/Plugins/Experimental/DigitalMultiplex/DMXProtocol/Source/DMXProtocol/Public/Interfaces/IDMXProtocolInterface.h b/Engine/Plugins/Experimental/DigitalMultiplex/DMXProtocol/Source/DMXProtocol/Public/Interfaces/IDMXProtocolInterface.h
new file mode 100644
--- /dev/null
+++ b/Engine/Plugins/Experimental/DigitalMultiplex/DMXProtocol/Source/DMXProtocol/Public/Interfaces/IDMXProtocolInterface.h
@@ -0,0 +1,33 @@
+// Copyright 1998-2019 Epic Games, Inc. All Rights Reserved.
+
+#pragma once
+
+#include "DMXProtocolCommon.h"
+
+/**
+ * IO interfacing protocol interface
+ */
+class DMXPROTOCOL_API IDMXProtocolInterface
+{
+public:
+	virtual ~IDMXProtocolInterface() {}
+
+	virtual IDMXProtocol* GetProtocol() const = 0;
+};
+
+/**
+ * Holds serial io parameters
+ */
+class DMXPROTOCOL_API IDMXProtocolInterfaceUSB
+	: public IDMXProtocolInterface
+{
+};
+
+/**
+ * Holds network io parameters
+ */
+class DMXPROTOCOL_API IDMXProtocolInterfaceEthernet
+	: public IDMXProtocolInterface
+{
+};
+
diff --git a/Engine/Plugins/Experimental/DigitalMultiplex/DMXProtocol/Source/DMXProtocol/Public/Interfaces/IDMXProtocolPacket.h b/Engine/Plugins/Experimental/DigitalMultiplex/DMXProtocol/Source/DMXProtocol/Public/Interfaces/IDMXProtocolPacket.h
new file mode 100644
--- /dev/null
+++ b/Engine/Plugins/Experimental/DigitalMultiplex/DMXProtocol/Source/DMXProtocol/Public/Interfaces/IDMXProtocolPacket.h
@@ -0,0 +1,13 @@
+// Copyright 1998-2019 Epic Games, Inc. All Rights Reserved.
+
+#pragma once
+
+#include "DMXProtocolCommon.h"
+
+struct IDMXProtocolPacket
+{
+public:
+	virtual ~IDMXProtocolPacket() {}
+
+	virtual TSharedPtr<FBufferArchive> Pack() { return nullptr; };
+};
diff --git a/Engine/Plugins/Experimental/DigitalMultiplex/DMXProtocol/Source/DMXProtocol/Public/Interfaces/IDMXProtocolPort.h b/Engine/Plugins/Experimental/DigitalMultiplex/DMXProtocol/Source/DMXProtocol/Public/Interfaces/IDMXProtocolPort.h
new file mode 100644
--- /dev/null
+++ b/Engine/Plugins/Experimental/DigitalMultiplex/DMXProtocol/Source/DMXProtocol/Public/Interfaces/IDMXProtocolPort.h
@@ -0,0 +1,37 @@
+// Copyright 1998-2019 Epic Games, Inc. All Rights Reserved.
+
+#pragma once
+
+#include "DMXProtocolCommon.h"
+
+/**
+ * Port interface
+ *
+ * Ports represent a single universe of DMX512. They are either input (receive
+ * DMX) or output (send DMX) but not both.
+ *
+ * Every port is part of a Device. Ports can be associated (patched) to a
+ * universe.
+ */
+class DMXPROTOCOL_API IDMXProtocolPort
+{
+public:
+	virtual ~IDMXProtocolPort() {}
+
+	virtual IDMXProtocol* GetProtocol() const = 0;
+
+	virtual uint8 GetPortID() const = 0;
+	virtual uint8 GetPriority() const = 0;
+	virtual void SetPriotiy(uint8 InPriority) = 0;
+	virtual bool IsSupportRDM() const = 0;
+	virtual TWeakPtr<IDMXProtocolDevice> GetCachedDevice() const = 0;
+	virtual TWeakPtr<IDMXProtocolUniverse> GetCachedUniverse() const = 0;
+	virtual void SetUniverse(const TSharedPtr<IDMXProtocolUniverse>& InUniverse) = 0;
+	virtual EDMXPortCapability GetPortCapability() const = 0;
+	virtual EDMXPortDirection GetPortDirection() const = 0;
+	virtual bool WriteDMX(const TSharedPtr<FDMXBuffer>& DMXBuffer) = 0;
+	virtual bool ReadDMX() = 0;
+	virtual TSharedPtr<FJsonObject> GetSettings() const = 0;
+	virtual uint16 GetUniverseID() const = 0;
+	virtual uint8 GetPortAddress() const { return 0x00; };
+};
diff --git a/Engine/Plugins/Experimental/DigitalMultiplex/DMXProtocol/Source/DMXProtocol/Public/Interfaces/IDMXProtocolRDM.h b/Engine/Plugins/Experimental/DigitalMultiplex/DMXProtocol/Source/DMXProtocol/Public/Interfaces/IDMXProtocolRDM.h
new file mode 100644
--- /dev/null
+++ b/Engine/Plugins/Experimental/DigitalMultiplex/DMXProtocol/Source/DMXProtocol/Public/Interfaces/IDMXProtocolRDM.h
@@ -0,0 +1,26 @@
+// Copyright 1998-2019 Epic Games, Inc. All Rights Reserved.
+
+#pragma once
+
+#include "DMXProtocolCommon.h"
+
+/**
+ */
+class DMXPROTOCOL_API IDMXProtocolRDM
+{
+public:
+	virtual ~IDMXProtocolRDM() {}
+
+	/**
+	 * Cross-protocol RDM command sending interface
+	 * @param CMD Key->Value command to send
+	 */
+	virtual void SendRDMCommand(const TSharedPtr<FJsonObject>& CMD) = 0;
+
+	/**
+	 * Cross-protocol RDM discovery sending interface
+	 * @param CMD Key->Value command to send
+	 */
+	virtual void RDMDiscovery(const TSharedPtr<FJsonObject>& CMD) = 0;
+};
+
diff --git a/Engine/Plugins/Experimental/DigitalMultiplex/DMXProtocol/Source/DMXProtocol/Public/Interfaces/IDMXProtocolTransport.h b/Engine/Plugins/Experimental/DigitalMultiplex/DMXProtocol/Source/DMXProtocol/Public/Interfaces/IDMXProtocolTransport.h
new file mode 100644
--- /dev/null
+++ b/Engine/Plugins/Experimental/DigitalMultiplex/DMXProtocol/Source/DMXProtocol/Public/Interfaces/IDMXProtocolTransport.h
@@ -0,0 +1,35 @@
+// Copyright 1998-2019 Epic Games, Inc. All Rights Reserved.
+
+#pragma once
+
+#include "DMXProtocolCommon.h"
+#include "HAL/Runnable.h"
+#include "HAL/ThreadSafeBool.h"
+#include "Misc/SingleThreadRunnable.h"
+#include "Serialization/ArrayWriter.h"
+
+class FRunnableThread;
+
+/**
+ * Delegate type for received data.
+ */
+DECLARE_DELEGATE_OneParam(FOnDMXDataReceived, const FArrayReaderPtr&);
+
+class IDMXProtocolSender 
+	: public FRunnable
+	, public FSingleThreadRunnable 
+{
+public:
+	virtual bool EnqueueOutboundPackage(FDMXPacketPtr Packet) = 0;
+};
+
+class IDMXProtocolReceiver
+	: public FRunnable
+	, public FSingleThreadRunnable
+{
+public:
+	virtual FOnDMXDataReceived& OnDataReceived() = 0;
+
+	virtual FRunnableThread* GetThread() const = 0;
+};
+
diff --git a/Engine/Plugins/Experimental/DigitalMultiplex/DMXProtocol/Source/DMXProtocol/Public/Interfaces/IDMXProtocolUniverse.h b/Engine/Plugins/Experimental/DigitalMultiplex/DMXProtocol/Source/DMXProtocol/Public/Interfaces/IDMXProtocolUniverse.h
new file mode 100644
--- /dev/null
+++ b/Engine/Plugins/Experimental/DigitalMultiplex/DMXProtocol/Source/DMXProtocol/Public/Interfaces/IDMXProtocolUniverse.h
@@ -0,0 +1,19 @@
+// Copyright 1998-2019 Epic Games, Inc. All Rights Reserved.
+
+#pragma once
+
+#include "DMXProtocolCommon.h"
+
+class DMXPROTOCOL_API IDMXProtocolUniverse
+{
+public:
+	virtual ~IDMXProtocolUniverse() {}
+
+	virtual IDMXProtocol* GetProtocol() const = 0;
+	virtual TSharedPtr<FDMXBuffer> GetOutputDMXBuffer() const = 0;
+	virtual TSharedPtr<FDMXBuffer> GetInputDMXBuffer() const = 0;
+	virtual bool SetDMXFragment(const IDMXFragmentMap& DMXFragment) = 0;
+	virtual uint8 GetPriority() const = 0;
+	virtual uint16 GetUniverseID() const = 0;
+	virtual TWeakPtr<IDMXProtocolPort> GetCachedUniversePort() const = 0;
+};
diff --git a/Engine/Plugins/Experimental/DigitalMultiplex/DMXProtocol/Source/DMXProtocol/Public/Managers/DMXProtocolDeviceManager.h b/Engine/Plugins/Experimental/DigitalMultiplex/DMXProtocol/Source/DMXProtocol/Public/Managers/DMXProtocolDeviceManager.h
new file mode 100644
--- /dev/null
+++ b/Engine/Plugins/Experimental/DigitalMultiplex/DMXProtocol/Source/DMXProtocol/Public/Managers/DMXProtocolDeviceManager.h
@@ -0,0 +1,25 @@
+// Copyright 1998-2019 Epic Games, Inc. All Rights Reserved.
+
+#pragma once
+
+#include "DMXProtocolCommon.h"
+#include "Interfaces/IDMXProtocol.h"
+
+class DMXPROTOCOL_API FDMXProtocolDeviceManager
+{
+public:
+	FDMXProtocolDeviceManager(IDMXProtocol* InDMXProtocol)
+		: DMXProtocol(InDMXProtocol)
+	{}
+
+public:
+	void AddDevice(TSharedPtr<IDMXProtocolDevice> InDevice);
+	void RemoveDevice(TSharedPtr<IDMXProtocolDevice> InDevice);
+	void RemoveAll();
+
+	bool GetDevicesByProtocol(const FName& ProtocolName, IDevicesMap& OutDevicesMap);
+
+private:
+	IDevicesMap DevicesMap;
+	IDMXProtocol* DMXProtocol;
+};
diff --git a/Engine/Plugins/Experimental/DigitalMultiplex/DMXProtocol/Source/DMXProtocol/Public/Managers/DMXProtocolInterfaceManager.h b/Engine/Plugins/Experimental/DigitalMultiplex/DMXProtocol/Source/DMXProtocol/Public/Managers/DMXProtocolInterfaceManager.h
new file mode 100644
--- /dev/null
+++ b/Engine/Plugins/Experimental/DigitalMultiplex/DMXProtocol/Source/DMXProtocol/Public/Managers/DMXProtocolInterfaceManager.h
@@ -0,0 +1,29 @@
+// Copyright 1998-2019 Epic Games, Inc. All Rights Reserved.
+
+#pragma once
+
+#include "DMXProtocolCommon.h"
+#include "Interfaces/IDMXProtocol.h"
+
+/**
+ * There is 2 different types
+ * If it ethernet protocol device it the IP adderess of network card. you could specify IP address  of network card manually
+ * This is helpful if PC has multiple network interfaces
+ * It is USB device this is serial port
+ */
+class DMXPROTOCOL_API FDMXProtocolInterfaceManager
+{
+public:
+	FDMXProtocolInterfaceManager(IDMXProtocol* InDMXProtocol)
+		: DMXProtocol(InDMXProtocol)
+	{}
+
+public:
+	void AddInterface(int32 InInterfaceID, TSharedPtr<IDMXProtocolInterface> InInterface);
+	void RemoveInterface(int32 InInterfaceID);
+	void RemoveAll();
+
+private:
+	IInterfacesMap InterfacesMap;
+	IDMXProtocol* DMXProtocol;
+};
diff --git a/Engine/Plugins/Experimental/DigitalMultiplex/DMXProtocol/Source/DMXProtocol/Public/Managers/DMXProtocolPortManager.h b/Engine/Plugins/Experimental/DigitalMultiplex/DMXProtocol/Source/DMXProtocol/Public/Managers/DMXProtocolPortManager.h
new file mode 100644
--- /dev/null
+++ b/Engine/Plugins/Experimental/DigitalMultiplex/DMXProtocol/Source/DMXProtocol/Public/Managers/DMXProtocolPortManager.h
@@ -0,0 +1,33 @@
+// Copyright 1998-2019 Epic Games, Inc. All Rights Reserved.
+
+#pragma once
+
+#include "DMXProtocolCommon.h"
+#include "Interfaces/IDMXProtocol.h"
+
+class DMXPROTOCOL_API FDMXProtocolPortManager
+{
+public:
+	FDMXProtocolPortManager(IDMXProtocol* InDMXProtocol)
+		: DMXProtocol(InDMXProtocol)
+	{}
+
+public:
+	void AddInputPort(TSharedPtr<IDMXProtocolDevice> InDevice, TSharedPtr<IDMXProtocolPort> InInputPort);
+	void AddOutputPort(TSharedPtr<IDMXProtocolDevice> InDevice, TSharedPtr<IDMXProtocolPort> InOutputPort);
+
+	void RemoveInputPort(TSharedPtr<IDMXProtocolDevice> InDevice, TSharedPtr<IDMXProtocolPort> InInputPort);
+	void RemoveOutputPort(TSharedPtr<IDMXProtocolDevice> InDevice, TSharedPtr<IDMXProtocolPort> InOutputPort);
+
+	IPortsMap* GetInputPortMapByDevice(IDMXProtocolDevice* InDevice);
+	IPortsMap* GetOutputPortMapByDevice(IDMXProtocolDevice* InDevice);
+
+	void RemoveAll();
+
+	IDMXProtocolPort* GetPortByDeviceAndID(const TSharedPtr<IDMXProtocolDevice>& InDevice, uint8 InPortID, EDMXPortDirection Direction);
+
+private:
+	IPortsDeviceMap IInputPortsDeviceMap;
+	IPortsDeviceMap IOutputPortsDeviceMap;
+	IDMXProtocol* DMXProtocol;
+};
diff --git a/Engine/Plugins/Experimental/DigitalMultiplex/DMXProtocol/Source/DMXProtocol/Public/Managers/DMXProtocolUniverseManager.h b/Engine/Plugins/Experimental/DigitalMultiplex/DMXProtocol/Source/DMXProtocol/Public/Managers/DMXProtocolUniverseManager.h
new file mode 100644
--- /dev/null
+++ b/Engine/Plugins/Experimental/DigitalMultiplex/DMXProtocol/Source/DMXProtocol/Public/Managers/DMXProtocolUniverseManager.h
@@ -0,0 +1,26 @@
+// Copyright 1998-2019 Epic Games, Inc. All Rights Reserved.
+
+#pragma once
+
+#include "DMXProtocolCommon.h"
+#include "Interfaces/IDMXProtocol.h"
+
+class DMXPROTOCOL_API FDMXProtocolUniverseManager
+{
+public:
+	FDMXProtocolUniverseManager(IDMXProtocol* InDMXProtocol)
+		: DMXProtocol(InDMXProtocol)
+	{}
+
+public:
+	void AddUniverse(TSharedPtr<IDMXProtocolPort> InPort, TSharedPtr<IDMXProtocolUniverse>  InUniverse);
+	void RemoveUniverse(TSharedPtr<IDMXProtocolPort> InPort, TSharedPtr<IDMXProtocolUniverse>  InUniverse);
+	void RemoveAll();
+	IDMXProtocolUniverse* GetUniverseByPort(IDMXProtocolPort* InPort);
+	IDMXProtocolUniverse* GetUniverseById(uint16 UniverseID);
+
+private:
+	IUniversesMap UniversesMap;
+	IUniversesIDMap UniversesIDMap;
+	IDMXProtocol* DMXProtocol;
+};
diff --git a/Engine/Plugins/Experimental/DigitalMultiplex/DMXProtocol/Source/DMXProtocol/Public/Serialization/DMXMemoryWriter.h b/Engine/Plugins/Experimental/DigitalMultiplex/DMXProtocol/Source/DMXProtocol/Public/Serialization/DMXMemoryWriter.h
new file mode 100644
--- /dev/null
+++ b/Engine/Plugins/Experimental/DigitalMultiplex/DMXProtocol/Source/DMXProtocol/Public/Serialization/DMXMemoryWriter.h
@@ -0,0 +1,18 @@
+#pragma once
+
+#include "Serialization/BufferArchive.h"
+
+class FDMXMemoryWriter : public FBufferArchive
+{
+public:
+	virtual void Serialize(void* Data, int64 Num) override
+	{
+		if (IsByteSwapping())
+		{
+			ByteSwap(Data, Num);
+		}
+
+		// Call Parent
+		FBufferArchive::Serialize(Data, Num);
+	}
+};
\ No newline at end of file
diff --git a/Engine/Plugins/Experimental/DigitalMultiplex/DMXProtocol/Source/DMXProtocolArtNet/DMXProtocolArtNet.Build.cs b/Engine/Plugins/Experimental/DigitalMultiplex/DMXProtocol/Source/DMXProtocolArtNet/DMXProtocolArtNet.Build.cs
new file mode 100644
--- /dev/null
+++ b/Engine/Plugins/Experimental/DigitalMultiplex/DMXProtocol/Source/DMXProtocolArtNet/DMXProtocolArtNet.Build.cs
@@ -0,0 +1,32 @@
+// Copyright 1998-2019 Epic Games, Inc. All Rights Reserved.
+
+using UnrealBuildTool;
+
+public class DMXProtocolArtNet : ModuleRules
+{
+	public DMXProtocolArtNet(ReadOnlyTargetRules Target) : base(Target)
+	{
+		PCHUsage = ModuleRules.PCHUsageMode.UseExplicitOrSharedPCHs;
+		
+		PublicDependencyModuleNames.AddRange(
+			new string[]
+			{
+				"Core",
+                "DMXProtocol"
+            }
+            );
+			
+		
+		PrivateDependencyModuleNames.AddRange(
+			new string[]
+			{
+				"CoreUObject",
+				"Engine",
+                "Networking",
+                "Sockets",
+                "Json"
+			}
+			);
+		
+	}
+}
diff --git a/Engine/Plugins/Experimental/DigitalMultiplex/DMXProtocol/Source/DMXProtocolArtNet/Private/DMXProtocolArtNet.cpp b/Engine/Plugins/Experimental/DigitalMultiplex/DMXProtocol/Source/DMXProtocolArtNet/Private/DMXProtocolArtNet.cpp
new file mode 100644
--- /dev/null
+++ b/Engine/Plugins/Experimental/DigitalMultiplex/DMXProtocol/Source/DMXProtocolArtNet/Private/DMXProtocolArtNet.cpp
@@ -0,0 +1,383 @@
+#include "DMXProtocolArtNet.h"
+#include "DMXProtocolTransportArtNet.h"
+#include "Common/UdpSocketBuilder.h"
+
+#include "DMXProtocolSettings.h"
+#include "Managers/DMXProtocolUniverseManager.h"
+#include "Interfaces/IDMXProtocolUniverse.h"
+#include "DMXProtocolArtNetUtils.h"
+#include "DMXProtocolPortArtNet.h"
+#include "DMXProtocolPackager.h"
+
+using namespace ArtNet;
+
+TSharedPtr<IDMXProtocolSender> FDMXProtocolArtNet::GetSenderInterface() const
+{
+	return ArtNetSender;
+}
+
+bool FDMXProtocolArtNet::Init()
+{
+	ISocketSubsystem* SocketSubsystem = ISocketSubsystem::Get(PLATFORM_SOCKETSUBSYSTEM);
+
+	// Create broadcast address
+	BroadcastAddr = SocketSubsystem->CreateInternetAddr();
+	BroadcastAddr->SetBroadcastAddress();
+	BroadcastAddr->SetPort(ARTNET_PORT);
+	BroadcastEndpoint = FIPv4Endpoint(BroadcastAddr);
+
+	// Create sender address
+	SenderAddr = SocketSubsystem->CreateInternetAddr();
+	bool bIsValid;
+	FString InterfaceIPAddress = GetDefault<UDMXProtocolSettings>()->InterfaceIPAddress;
+	SenderAddr->SetIp(*InterfaceIPAddress, bIsValid);
+	SenderAddr->SetPort(ARTNET_PORT);
+	SenderEndpoint = FIPv4Endpoint(SenderAddr);
+
+	// Create Broadcast socket
+	BroadcastSocket = FUdpSocketBuilder(TEXT("UDPArtNetBroadcastSocket"))
+		.AsNonBlocking()
+		.AsReusable()
+		.WithBroadcast()
+#if PLATFORM_WINDOWS
+		.BoundToAddress(SenderEndpoint.Address)
+#endif
+		.BoundToPort(SenderEndpoint.Port);
+
+	if (BroadcastSocket == nullptr)
+	{
+		UE_LOG_DMXPROTOCOL(Error, TEXT("ERROR create BroadcastSocket"));
+	}
+
+	// Create sender
+	ArtNetSender = MakeShared<FDMXProtocolSenderArtNet>(*BroadcastSocket, this);
+
+	// Create receiver
+	FTimespan ThreadWaitTime = FTimespan::FromMilliseconds(100);
+	ArtNetReceiver = MakeShared<FDMXProtocolReceiverArtNet>(*BroadcastSocket, this, ThreadWaitTime);
+	ArtNetReceiver->OnDataReceived().BindRaw(this, &FDMXProtocolArtNet::OnDataReceived);
+
+	return true;
+}
+
+bool FDMXProtocolArtNet::Shutdown()
+{
+	ArtNetSender.Reset();
+	ArtNetReceiver.Reset();
+
+	//Clear all sockets!
+	if (BroadcastSocket)
+	{
+		BroadcastSocket->Close();
+		ISocketSubsystem::Get(PLATFORM_SOCKETSUBSYSTEM)->DestroySocket(BroadcastSocket);
+	}
+
+	return false;
+}
+
+bool FDMXProtocolArtNet::IsEnabled() const
+{
+	return true;
+}
+
+FName FDMXProtocolArtNet::GetProtocolName() const
+{
+	return FDMXProtocolArtNetModule::NAME_Artnet;
+}
+
+void FDMXProtocolArtNet::Reload()
+{
+}
+
+bool FDMXProtocolArtNet::Tick(float DeltaTime)
+{
+	return true;
+}
+
+bool FDMXProtocolArtNet::SendDMX(uint16 UniverseID, uint8 PortID, const TSharedPtr<FDMXBuffer>& DMXBuffer) const
+{
+	// Init Packager
+	FDMXProtocolPackager Packager;
+
+	// ArtNet DMX packet
+	FDMXProtocolArtNetDMXPacket ArtNetDMXPacket;
+	if (DMXBuffer->GetDMXData().Num() == ARTNET_DMX_LENGTH)
+	{
+		FMemory::Memcpy(ArtNetDMXPacket.Data, DMXBuffer->GetDMXData().GetData(), ARTNET_DMX_LENGTH);
+	}
+	else
+	{
+		UE_LOG_DMXPROTOCOL(Warning, TEXT("Size of outcoming ArtNet DMX buffer is wrong"));
+	}
+	ArtNetDMXPacket.Physical = PortID;
+	ArtNetDMXPacket.Universe = UniverseID;
+	Packager.AddToPackage(&ArtNetDMXPacket);
+
+	// Sending
+	FDMXPacketPtr Packet = MakeShared<FDMXPacket, ESPMode::ThreadSafe>(Packager.GetBuffer());
+	GetSenderInterface()->EnqueueOutboundPackage(Packet);
+
+	return false;
+}
+
+void FDMXProtocolArtNet::SendRDMCommand(const TSharedPtr<FJsonObject>& CMD)
+{
+	// Cross protocol RDM send command implementation will be here
+}
+
+void FDMXProtocolArtNet::RDMDiscovery(const TSharedPtr<FJsonObject>& CMD)
+{
+	// Cross protocol RDM discovery implementation will be here
+}
+
+bool FDMXProtocolArtNet::TransmitPoll()
+{
+	// Init Packager
+	FDMXProtocolPackager Packager;
+
+	// Poll packet
+	FDMXProtocolArtNetPollPacket PollPacket;
+	Packager.AddToPackage(&PollPacket);
+
+	// Sending
+	FDMXPacketPtr Packet = MakeShared<FDMXPacket, ESPMode::ThreadSafe>(Packager.GetBuffer());
+	return GetSenderInterface()->EnqueueOutboundPackage(Packet);
+}
+
+bool FDMXProtocolArtNet::TransmitTodRequestToAll()
+{
+	bool bIsSent = false;
+	IDevicesMap DevicesMap;
+	if (GetDeviceManager()->GetDevicesByProtocol(GetProtocolName(), DevicesMap))
+	{
+		for (TMap<IDMXProtocolInterface*, TSharedPtr<IDMXProtocolDevice>>::TConstIterator DeviceIt = DevicesMap.CreateConstIterator(); DeviceIt; ++DeviceIt)
+		{
+			if (DeviceIt->Value.IsValid())
+			{
+				TransmitTodRequest(DeviceIt->Value);
+				bIsSent = true;
+			}
+		}
+	}
+
+	return bIsSent;
+}
+
+bool FDMXProtocolArtNet::TransmitTodRequest(const TSharedPtr<IDMXProtocolDevice>& InDevice)
+{
+	// Init Packager
+	FDMXProtocolPackager Packager;
+
+	// TransmitTodRequest packet
+	FDMXProtocolArtNetTodRequest TodRequest;
+
+	if (InDevice.IsValid())
+	{
+		const IPortsMap* OutputPortsMapDevice = GetPortManager()->GetOutputPortMapByDevice(InDevice.Get());
+		for (uint8 PortIndex = 0; PortIndex < ARTNET_MAX_PORTS; PortIndex++)
+		{
+			if (OutputPortsMapDevice != nullptr)
+			{
+				const TSharedPtr<IDMXProtocolPort>* Port = OutputPortsMapDevice->Find(PortIndex);
+				if (Port != nullptr && Port->IsValid())
+				{
+					FDMXProtocolPortArtNet* PortArtNet = static_cast<FDMXProtocolPortArtNet*>(Port->Get());
+					TodRequest.Address[TodRequest.AdCount] = PortArtNet->GetPortAddress();
+					TodRequest.AdCount++;
+				}
+			}
+		}
+	}
+
+	Packager.AddToPackage(&TodRequest);
+
+	// Sending
+	FDMXPacketPtr Packet = MakeShared<FDMXPacket, ESPMode::ThreadSafe>(Packager.GetBuffer());
+	return GetSenderInterface()->EnqueueOutboundPackage(Packet);
+}
+
+bool FDMXProtocolArtNet::TransmitTodData(const TSharedPtr<IDMXProtocolDevice>& InDevice, uint8 InPortID)
+{
+	// Init Packager
+	FDMXProtocolPackager Packager;
+
+	// TransmitTodData packet
+	FDMXProtocolArtNetTodData TodData;
+	TodData.Port = InPortID;
+
+	FDMXProtocolPortArtNet* PortArtNet = static_cast<FDMXProtocolPortArtNet*>(GetPortManager()->GetPortByDeviceAndID(InDevice, InPortID, EDMXPortDirection::DMX_PORT_OUTPUT));
+	if (PortArtNet != nullptr)
+	{
+		TodData.Net = PortArtNet->GetNetAddress();
+		TodData.Address = PortArtNet->GetUniverseAddress();
+
+		const TArray<FRDMUID>& TODUID = PortArtNet->GetTODUIDs();
+		const uint32 NumTODBytes = TODUID.Num() * ARTNET_RDM_UID_WIDTH;
+
+		TodData.UidTotalHi = DMX_SHORT_GET_HIGH_BIT(TODUID.Num());
+		TodData.UidTotal = DMX_SHORT_GET_LOW_BYTE(TODUID.Num());
+
+		if (TODUID.Num())
+		{
+			TodData.BlockCount = 1;
+			TodData.UidCount = TODUID.Num();
+			FMemory::Memcpy(TodData.Tod, TODUID.GetData(), NumTODBytes);
+		}
+	}
+
+	Packager.AddToPackage(&TodData);
+
+	// Sending
+	FDMXPacketPtr Packet = MakeShared<FDMXPacket, ESPMode::ThreadSafe>(Packager.GetBuffer());
+	return GetSenderInterface()->EnqueueOutboundPackage(Packet);
+}
+
+bool FDMXProtocolArtNet::TransmitTodControl(const TSharedPtr<IDMXProtocolDevice>& InDevice, uint8 InPortID, uint8 Action)
+{
+	// Try to get ArtNet Port
+	FDMXProtocolPortArtNet* PortArtNet = static_cast<FDMXProtocolPortArtNet*>(GetPortManager()->GetPortByDeviceAndID(InDevice, InPortID, EDMXPortDirection::DMX_PORT_OUTPUT));
+
+	// Init Packager
+	FDMXProtocolPackager Packager;
+
+	// TransmitTodControl packet
+	FDMXProtocolArtNetTodControl TodControl;
+	TodControl.Cmd = Action;
+
+	if (PortArtNet != nullptr)
+	{
+		TodControl.Net = PortArtNet->GetNetAddress();
+		TodControl.Address = PortArtNet->GetPortAddress();
+	}
+
+	Packager.AddToPackage(&TodControl);
+
+	// Sending
+	FDMXPacketPtr Packet = MakeShared<FDMXPacket, ESPMode::ThreadSafe>(Packager.GetBuffer());
+	return GetSenderInterface()->EnqueueOutboundPackage(Packet);
+}
+
+bool FDMXProtocolArtNet::TransmitRDM(const TSharedPtr<IDMXProtocolDevice>& InDevice, uint8 InPortID, const TArray<uint8>& Data)
+{
+	// Try to get ArtNet Port
+	FDMXProtocolPortArtNet* PortArtNet = static_cast<FDMXProtocolPortArtNet*>(GetPortManager()->GetPortByDeviceAndID(InDevice, InPortID, EDMXPortDirection::DMX_PORT_OUTPUT));
+
+	// Init Packager
+	FDMXProtocolPackager Packager;
+
+	// TransmitRDM packet
+	FDMXProtocolArtNetRDM RDM;
+	if (PortArtNet != nullptr)
+	{
+		RDM.Net = PortArtNet->GetNetAddress();
+		RDM.Address = PortArtNet->GetUniverseAddress();
+	}
+
+	if (Data.Num() <= ARTNET_MAX_RDM_DATA)
+	{
+		FMemory::Memcpy(RDM.Data, Data.GetData(), ARTNET_MAX_RDM_DATA);
+	}
+	else
+	{
+		UE_LOG_DMXPROTOCOL(Warning, TEXT("Size of outcoming ArtNet RDM commands is bigger then limit"));
+	}
+
+	Packager.AddToPackage(&RDM);
+
+	// Sending
+	FDMXPacketPtr Packet = MakeShared<FDMXPacket, ESPMode::ThreadSafe>(Packager.GetBuffer());
+	return GetSenderInterface()->EnqueueOutboundPackage(Packet);
+}
+
+void FDMXProtocolArtNet::OnDataReceived(const FArrayReaderPtr& Buffer)
+{
+	// It will be more handlers
+	switch (ArtNet::GetPacketType(Buffer))
+	{
+	case ARTNET_REPLY:
+		HandleReplyPacket(Buffer);
+		break;
+	case ARTNET_DMX:
+		HandleDataPacket(Buffer);
+		break;
+	case ARTNET_TODREQUEST:
+		HandleTodRequest(Buffer);
+		break;
+	case ARTNET_TODDATA:
+		HandleTodData(Buffer);
+		break;
+	case ARTNET_TODCONTROL:
+		HandleTodControl(Buffer);
+		break;
+	case ARTNET_RDM:
+		HandleRdm(Buffer);
+	default:
+		break;
+	}
+}
+
+bool FDMXProtocolArtNet::HandleReplyPacket(const FArrayReaderPtr & Buffer)
+{
+	FArtNetPacketReply ArtNetPacketReply;
+
+	*Buffer << ArtNetPacketReply;
+
+	// Nothing to do with packet now
+	return false;
+}
+
+bool FDMXProtocolArtNet::HandleDataPacket(const FArrayReaderPtr & Buffer)
+{
+	// ArtNet DMX packet
+	FDMXProtocolArtNetDMXPacket ArtNetDMXPacket;
+	*Buffer << ArtNetDMXPacket;
+
+	// Write data to input DMX buffer universe
+	IDMXProtocolUniverse* Universe = GetUniverseManager()->GetUniverseById(ArtNetDMXPacket.Universe);
+	if (Universe != nullptr && Universe->GetInputDMXBuffer().IsValid())
+	{
+		// Make sure we copy same amount of data
+		if (Universe->GetInputDMXBuffer()->GetDMXData().Num() == ARTNET_DMX_LENGTH)
+		{
+			FMemory::Memcpy(Universe->GetInputDMXBuffer()->GetDMXData().GetData(), ArtNetDMXPacket.Data, ARTNET_DMX_LENGTH);
+			return true;
+		}
+		else
+		{
+			UE_LOG_DMXPROTOCOL(Warning, TEXT("Size of incoming DMX buffer is wrong"));
+			return false;
+		}
+	}
+
+	return false;
+}
+
+bool FDMXProtocolArtNet::HandleTodRequest(const FArrayReaderPtr & Buffer)
+{
+	*Buffer << IncomingTodRequest;
+
+	return true;
+}
+
+bool FDMXProtocolArtNet::HandleTodData(const FArrayReaderPtr & Buffer)
+{
+	*Buffer << IncomingTodData;
+
+	return true;
+}
+
+bool FDMXProtocolArtNet::HandleTodControl(const FArrayReaderPtr & Buffer)
+{
+	*Buffer << IncomingTodControl;
+
+	return true;
+}
+
+bool FDMXProtocolArtNet::HandleRdm(const FArrayReaderPtr & Buffer)
+{
+	*Buffer << IncomingRDM;
+
+	return true;
+}
+
+
diff --git a/Engine/Plugins/Experimental/DigitalMultiplex/DMXProtocol/Source/DMXProtocolArtNet/Private/DMXProtocolArtNetConstants.h b/Engine/Plugins/Experimental/DigitalMultiplex/DMXProtocol/Source/DMXProtocolArtNet/Private/DMXProtocolArtNetConstants.h
new file mode 100644
--- /dev/null
+++ b/Engine/Plugins/Experimental/DigitalMultiplex/DMXProtocol/Source/DMXProtocolArtNet/Private/DMXProtocolArtNetConstants.h
@@ -0,0 +1,140 @@
+// Copyright 1998-2019 Epic Games, Inc. All Rights Reserved.
+
+#pragma once
+
+#include "CoreMinimal.h"
+#include "DMXProtocolTypes.h"
+
+enum
+{
+	LOW_BYTE = 0x00FF,
+	HIGH_BYTE = 0xFF00
+};
+
+#ifndef DMX_SHORT_GET_HIGH_BIT
+#define DMX_SHORT_GET_HIGH_BIT(x) ((HIGH_BYTE & x) >> 8)
+#endif
+
+#ifndef DMX_SHORT_GET_LOW_BYTE
+#define DMX_SHORT_GET_LOW_BYTE(x)  (LOW_BYTE & x)
+#endif
+
+enum
+{
+	ARTNET_LONG_NAME_LENGTH = 64,
+	ARTNET_MAX_PORTS = 4,
+	ARTNET_MAX_RDM_ADDRESS_COUNT = 32,
+	ARTNET_MAX_RDM_DATA = 512,
+	ARTNET_MAX_UID_COUNT = 200,
+	ARTNET_REPORT_LENGTH = 64,
+	ARTNET_SHORT_NAME_LENGTH = 18
+};
+
+enum {
+	ARTNET_POLL = 0x2000,
+	ARTNET_REPLY = 0x2100,
+	ARTNET_DMX = 0x5000,
+	ARTNET_ADDRESS = 0x6000,
+	ARTNET_INPUT = 0x7000,
+	ARTNET_TODREQUEST = 0x8000,
+	ARTNET_TODDATA = 0x8100,
+	ARTNET_TODCONTROL = 0x8200,
+	ARTNET_RDM = 0x8300,
+	ARTNET_VIDEOSTEUP = 0xa010,
+	ARTNET_VIDEOPALETTE = 0xa020,
+	ARTNET_VIDEODATA = 0xa040,
+	ARTNET_MACMASTER = 0xf000,
+	ARTNET_MACSLAVE = 0xf100,
+	ARTNET_FIRMWAREMASTER = 0xf200,
+	ARTNET_FIRMWAREREPLY = 0xf300,
+	ARTNET_IPPROG = 0xf800,
+	ARTNET_IPREPLY = 0xf900,
+	ARTNET_MEDIA = 0x9000,
+	ARTNET_MEDIAPATCH = 0x9200,
+	ARTNET_MEDIACONTROLREPLY = 0x9300
+};
+
+enum {
+	ARTNET_DMX_LENGTH = 512,
+	ARTNET_DMX_PACKAGE_SIZE = 530,
+	ARTNET_POLL_PACKAGE_SIZE = 14,
+};
+
+enum
+{
+	ARTNET_MAX_UNIVERSES = 32767
+};
+
+enum  
+{
+	ARTNET_TTM_DEFAULT = 0xFF,  /**< default, ArtPollReplies are broadcast, and nodes won't send a ArtPollReply when conditions change */
+	ARTNET_TTM_PRIVATE = 0xFE,  /**< ArtPollReplies aren't broadcast */
+	ARTNET_TTM_noauto = 0xFD    /**< ArtPollReplies are send when node conditions chang */
+};
+
+// TOD actions
+typedef enum {
+	ARTNET_TOD_FULL = 0x00,
+	ARTNET_TOD_FLUSH = 0x01,
+} artnet_tod_command_code;
+
+
+// RDM
+enum
+{
+	ARTNET_RDM_VERSION = 0x01,
+	ARTNET_MAX_RDM_ADCOUNT = 32,
+
+	/**
+	 * ARTNET_RDM_UID_WIDTH are 6 bytes, the first two bytes are the
+	 * [manufacturer code](http://tsp.esta.org/tsp/working_groups/CP/mfctrIDs.php)
+	 * and the last 4 bytes are the device id. UIDs are written as: XXXX:YYYYYYYY
+	 */
+	ARTNET_RDM_UID_WIDTH = 6,
+};
+
+// RDM functions
+enum
+{
+	ARTNET_RDM_PROCESS_PACKET = 0x00
+};
+
+enum
+{
+	ARTNET_ADDRESS_NO_CHANGE = 0x7f,
+	ARTNET_PORT = 6454,
+	ARTNET_STRING_SIZE = 8,
+	ARTNET_VERSION = 14,
+	OEM_HI = 0x04,
+	OEM_LO = 0x30,
+	ESTA_HI = 'z',
+	ESTA_LO = 'p',
+	TTM_BEHAVIOUR_MASK = 0x02,
+	TTM_REPLY_MASK = 0x01,
+	PROGRAM_NO_CHANGE = 0x7f,
+	PROGRAM_DEFAULTS = 0x00,
+	PROGRAM_CHANGE_MASK = 0x80,
+	STATUS_PROG_AUTH_MASK = 0x30,
+	PORT_STATUS_LPT_MODE = 0x02,
+	PORT_STATUS_SHORT = 0x04,
+	PORT_STATUS_ERROR = 0x04,
+	PORT_STATUS_DISABLED_MASK = 0x08,
+	PORT_STATUS_MERGE = 0x08,
+	PORT_STATUS_DMXEXT = 0x10,
+	PORT_STATUS_DMX_SIP = 0x20,
+	PORT_STATUS_DMXEST = 0x40,
+	PORT_STATUS_ACT_MASK = 0x80,
+	PORT_DISABLE_MASK = 0x01,
+	TOD_RESPONSE_FULL = 0x00,
+	TOD_RESPONSE_NAK = 0x00,
+	MIN_PACKET_SIZE = 10,
+	MERGEIMEOUT_SECONDS = 10,
+	FIRMWAREIMEOUT_SECONDS = 20,
+	RECV_NO_DATA = 1,
+	MAX_NODE_BCAST_LIMIT = 30 // always bcast after this point
+};
+
+struct DMXArtNet
+{
+	static uint8 constexpr ARTNET_STRING[ARTNET_STRING_SIZE] = "Art-Net";
+};
diff --git a/Engine/Plugins/Experimental/DigitalMultiplex/DMXProtocol/Source/DMXProtocolArtNet/Private/DMXProtocolArtNetModule.cpp b/Engine/Plugins/Experimental/DigitalMultiplex/DMXProtocol/Source/DMXProtocolArtNet/Private/DMXProtocolArtNetModule.cpp
new file mode 100644
--- /dev/null
+++ b/Engine/Plugins/Experimental/DigitalMultiplex/DMXProtocol/Source/DMXProtocolArtNet/Private/DMXProtocolArtNetModule.cpp
@@ -0,0 +1,123 @@
+// Copyright 1998-2019 Epic Games, Inc. All Rights Reserved.
+
+#include "DMXProtocolArtNetModule.h"
+#include "CoreMinimal.h"
+#include "Modules/ModuleManager.h"
+#include "Json.h"
+
+#include "DMXProtocolArtNet.h"
+#include "DMXProtocolTypes.h"
+
+
+IMPLEMENT_MODULE(FDMXProtocolArtNetModule, DMXProtocolArtNet);
+
+const FName FDMXProtocolArtNetModule::NAME_Artnet = FName(TEXT("ARTNET"));
+
+FAutoConsoleCommand FDMXProtocolArtNetModule::SendDMXCommand(
+	TEXT("DMX.ArtNet.SendDMX"),
+	TEXT("Command for sending DMX through ArtNet Protocol. DMX.ArtNet.SendDMX [UniverseID] Channel:Value Channel:Value Channel:Value n\t DMX.ArtNet.SendDMX 17 10:6 11:7 12:8 13:9 n\t It will send channels values to the DMX to Universe 17"),
+	FConsoleCommandWithArgsDelegate::CreateStatic(&FDMXProtocolArtNetModule::SendDMXCommandHandler)
+);
+
+TSharedPtr<IDMXProtocol> FDMXProtocolFactoryArtNet::CreateProtocol(const FName & ProtocolName)
+{
+	FJsonObject ProtocolSettings;
+	TSharedPtr<IDMXProtocol> ProtocolArtNetPtr = MakeShared<FDMXProtocolArtNet>(ProtocolName, ProtocolSettings);
+	if (ProtocolArtNetPtr->IsEnabled())
+	{
+		if (!ProtocolArtNetPtr->Init())
+		{
+			UE_LOG_DMXPROTOCOL(Warning, TEXT("ArtNet failed to initialize!"));
+			ProtocolArtNetPtr->Shutdown();
+			ProtocolArtNetPtr = nullptr;
+		}
+	}
+	else
+	{
+		UE_LOG_DMXPROTOCOL(Warning, TEXT("ArtNet disabled!"));
+		ProtocolArtNetPtr->Shutdown();
+		ProtocolArtNetPtr = nullptr;
+	}
+
+	return ProtocolArtNetPtr;
+}
+
+void FDMXProtocolArtNetModule::StartupModule()
+{
+	FactoryArtNet = MakeUnique<FDMXProtocolFactoryArtNet>();
+
+	// Create and register our singleton factory with the main online subsystem for easy access
+	FDMXProtocolModule& DMXProtocolModule = FModuleManager::GetModuleChecked<FDMXProtocolModule>("DMXProtocol");
+	DMXProtocolModule.RegisterProtocol(NAME_Artnet, FactoryArtNet.Get());
+}
+
+void FDMXProtocolArtNetModule::ShutdownModule()
+{
+	// Unregister and destroy protocol
+	FDMXProtocolModule* DMXProtocolModule = FModuleManager::GetModulePtr<FDMXProtocolModule>("DMXProtocol");
+	if (DMXProtocolModule != nullptr)
+	{
+		DMXProtocolModule->UnregisterProtocol(NAME_Artnet);
+	}
+	
+	FactoryArtNet.Release();
+}
+
+FDMXProtocolArtNetModule& FDMXProtocolArtNetModule::Get()
+{
+	return FModuleManager::GetModuleChecked<FDMXProtocolArtNetModule>("DMXProtocolArtNet");
+}
+
+void FDMXProtocolArtNetModule::SendDMXCommandHandler(const TArray<FString>& Args)
+{
+	if (Args.Num() < 2)
+	{
+		UE_LOG_DMXPROTOCOL(Warning, TEXT("Not enough arguments. It won't be sent\n"
+										"Command structure is DMX.ArtNet.SendDMX [UniverseID] Channel:Value Channel:Value Channel:Value\n"
+										"For example: DMX.ArtNet.SendDMX 17 10:6 11:7 12:8 13:9"));
+		return;
+	}
+
+	uint32 UniverseID = 0;
+	LexTryParseString<uint32>(UniverseID, *Args[0]);
+	if (UniverseID > ARTNET_MAX_UNIVERSES)
+	{
+		UE_LOG_DMXPROTOCOL(Warning, TEXT("The UniverseID is bigger then the max universe value. It won't be sent\n"
+			"For example: DMX.ArtNet.SendDMX 17 10:6 11:7 12:8 13:9\n"
+			"Where Universe %d should be less then %d"), UniverseID, ARTNET_MAX_UNIVERSES);
+		return;
+	}
+
+	IDMXFragmentMap DMXFragment;
+	for (int32 i = 1; i < Args.Num(); i++)
+	{
+		FString ChannelAndValue = Args[i];
+
+		FString KeyStr;
+		FString ValueStr;
+		ChannelAndValue.Split(TEXT(":"), &KeyStr, &ValueStr);
+
+		uint32 Key = 0;
+		uint32 Value = 0;
+		LexTryParseString<uint32>(Key, *KeyStr);
+		if (Key > DMX_UNIVERSE_SIZE)
+		{
+			UE_LOG_DMXPROTOCOL(Warning, TEXT("The input channel is bigger then the universe size. It won't be sent\n"
+				"For example: DMX.ArtNet.SendDMX 17 10:6 11:7 12:8 13:9\n"
+				"Where channel %d should be less then %d"), Key, DMX_UNIVERSE_SIZE);
+			return;
+		}
+		LexTryParseString<uint32>(Value, *ValueStr);
+		if (Value > DMX_MAX_CHANNEL_VALUE)
+		{
+			UE_LOG_DMXPROTOCOL(Warning, TEXT("The input value is bigger then the universe size. It won't be sent\n"
+				"For example: DMX.ArtNet.SendDMX 17 10:6 11:7 12:8 13:9\n"
+				"Where value %d should be less then %d"), Value, DMX_MAX_CHANNEL_VALUE);
+			return;
+		}
+		DMXFragment.Add(Key, Value);
+	}
+
+	FDMXProtocolArtNet* DMXProtocol = IDMXProtocol::Get<FDMXProtocolArtNet>(FDMXProtocolArtNetModule::NAME_Artnet);
+	DMXProtocol->SetDMXFragment(UniverseID, DMXFragment, true);
+}
diff --git a/Engine/Plugins/Experimental/DigitalMultiplex/DMXProtocol/Source/DMXProtocolArtNet/Private/DMXProtocolArtNetUtils.h b/Engine/Plugins/Experimental/DigitalMultiplex/DMXProtocol/Source/DMXProtocolArtNet/Private/DMXProtocolArtNetUtils.h
new file mode 100644
--- /dev/null
+++ b/Engine/Plugins/Experimental/DigitalMultiplex/DMXProtocol/Source/DMXProtocolArtNet/Private/DMXProtocolArtNetUtils.h
@@ -0,0 +1,27 @@
+// Copyright 1998-2019 Epic Games, Inc. All Rights Reserved.
+
+#pragma once
+
+#include "CoreMinimal.h"
+#include "DMXProtocolTypes.h"
+#include "DMXProtocolArtNetConstants.h"
+
+namespace ArtNet
+{
+	static uint16 GetPacketType(const FArrayReaderPtr& Buffer)
+	{
+		uint16 OpCode = 0x0000;
+		const uint32 MinCheck = ARTNET_STRING_SIZE + 2;
+		if (Buffer->Num() > MinCheck)
+		{
+			// Get OpCode
+			Buffer->Seek(ARTNET_STRING_SIZE);
+			*Buffer << OpCode;
+
+			// Reset Position
+			Buffer->Seek(0);
+		}
+
+		return OpCode;
+	}
+}
\ No newline at end of file
diff --git a/Engine/Plugins/Experimental/DigitalMultiplex/DMXProtocol/Source/DMXProtocolArtNet/Private/DMXProtocolDeviceArtNet.cpp b/Engine/Plugins/Experimental/DigitalMultiplex/DMXProtocol/Source/DMXProtocolArtNet/Private/DMXProtocolDeviceArtNet.cpp
new file mode 100644
--- /dev/null
+++ b/Engine/Plugins/Experimental/DigitalMultiplex/DMXProtocol/Source/DMXProtocolArtNet/Private/DMXProtocolDeviceArtNet.cpp
@@ -0,0 +1,43 @@
+#include "DMXProtocolDeviceArtNet.h"
+#include "Interfaces/IDMXProtocol.h"
+#include "Interfaces/IDMXProtocolInterface.h"
+#include "Interfaces/IDMXProtocolPort.h"
+#include "Managers/DMXProtocolPortManager.h"
+#include "DMXProtocolTypes.h"
+#include "Json.h"
+
+FDMXProtocolDeviceArtNet::FDMXProtocolDeviceArtNet(IDMXProtocol* InDMXProtocol, TSharedPtr<IDMXProtocolInterface> InProtocolInterface, FJsonObject& InSettings, uint32 InDeviceID)
+	: DMXProtocol(InDMXProtocol)
+	, ProtocolInterface(InProtocolInterface)
+	, DeviceID(InDeviceID)
+{
+	Settings = MakeShared<FJsonObject>(InSettings);
+
+	checkf(DMXProtocol, TEXT("DMXProtocol pointer is nullptr"));
+	checkf(ProtocolInterface.IsValid(), TEXT("ProtocolInterface is not valid"));
+}
+
+TSharedPtr<FJsonObject> FDMXProtocolDeviceArtNet::GetSettings() const
+{
+	return Settings;
+}
+
+TWeakPtr<IDMXProtocolInterface> FDMXProtocolDeviceArtNet::GetCachedProtocolInterface() const
+{
+	return ProtocolInterface;
+}
+
+IDMXProtocol* FDMXProtocolDeviceArtNet::GetProtocol() const
+{
+	return DMXProtocol;
+}
+
+uint32 FDMXProtocolDeviceArtNet::GetDeviceID() const
+{
+	return DeviceID;
+}
+
+bool FDMXProtocolDeviceArtNet::AllowLooping() const
+{
+	return true;
+}
diff --git a/Engine/Plugins/Experimental/DigitalMultiplex/DMXProtocol/Source/DMXProtocolArtNet/Private/DMXProtocolDeviceArtNet.h b/Engine/Plugins/Experimental/DigitalMultiplex/DMXProtocol/Source/DMXProtocolArtNet/Private/DMXProtocolDeviceArtNet.h
new file mode 100644
--- /dev/null
+++ b/Engine/Plugins/Experimental/DigitalMultiplex/DMXProtocol/Source/DMXProtocolArtNet/Private/DMXProtocolDeviceArtNet.h
@@ -0,0 +1,29 @@
+// Copyright 1998-2019 Epic Games, Inc. All Rights Reserved.
+
+#pragma once
+
+#include "CoreMinimal.h"
+#include "Interfaces/IDMXProtocolDevice.h"
+#include "Json.h"
+
+class DMXPROTOCOLARTNET_API FDMXProtocolDeviceArtNet
+	: public IDMXProtocolDevice
+{
+public:
+	FDMXProtocolDeviceArtNet(IDMXProtocol* InDMXProtocol, TSharedPtr<IDMXProtocolInterface> InProtocolInterface, FJsonObject& InSettings, uint32 InDeviceID);
+
+	//~ Begin IDMXProtocolDevice implementation
+	virtual TSharedPtr<FJsonObject> GetSettings() const override;
+	virtual TWeakPtr<IDMXProtocolInterface> GetCachedProtocolInterface() const override;
+	virtual IDMXProtocol* GetProtocol() const override;
+	virtual uint32 GetDeviceID() const override;
+
+	virtual bool AllowLooping() const override;
+	//~ End IDMXProtocolDevice implementation
+
+private:
+	IDMXProtocol* DMXProtocol;
+	TWeakPtr<IDMXProtocolInterface> ProtocolInterface;
+	TSharedPtr<FJsonObject> Settings;
+	uint32 DeviceID;
+};
\ No newline at end of file
diff --git a/Engine/Plugins/Experimental/DigitalMultiplex/DMXProtocol/Source/DMXProtocolArtNet/Private/DMXProtocolInterfaceArtNet.cpp b/Engine/Plugins/Experimental/DigitalMultiplex/DMXProtocol/Source/DMXProtocolArtNet/Private/DMXProtocolInterfaceArtNet.cpp
new file mode 100644
--- /dev/null
+++ b/Engine/Plugins/Experimental/DigitalMultiplex/DMXProtocol/Source/DMXProtocolArtNet/Private/DMXProtocolInterfaceArtNet.cpp
@@ -0,0 +1,12 @@
+#include "DMXProtocolInterfaceArtNet.h"
+#include "Interfaces/IDMXProtocol.h"
+
+FDMXProtocolInterfaceArtNet::FDMXProtocolInterfaceArtNet(IDMXProtocol* InDMXProtocol)
+	: DMXProtocol(InDMXProtocol)
+{
+}
+
+IDMXProtocol* FDMXProtocolInterfaceArtNet::GetProtocol() const
+{
+	return DMXProtocol;
+}
\ No newline at end of file
diff --git a/Engine/Plugins/Experimental/DigitalMultiplex/DMXProtocol/Source/DMXProtocolArtNet/Private/DMXProtocolInterfaceArtNet.h b/Engine/Plugins/Experimental/DigitalMultiplex/DMXProtocol/Source/DMXProtocolArtNet/Private/DMXProtocolInterfaceArtNet.h
new file mode 100644
--- /dev/null
+++ b/Engine/Plugins/Experimental/DigitalMultiplex/DMXProtocol/Source/DMXProtocolArtNet/Private/DMXProtocolInterfaceArtNet.h
@@ -0,0 +1,20 @@
+// Copyright 1998-2019 Epic Games, Inc. All Rights Reserved.
+
+#pragma once
+
+#include "CoreMinimal.h"
+#include "Interfaces/IDMXProtocolInterface.h"
+
+class DMXPROTOCOLARTNET_API FDMXProtocolInterfaceArtNet
+	: public IDMXProtocolInterfaceEthernet
+{
+public:
+	FDMXProtocolInterfaceArtNet(IDMXProtocol* InDMXProtocol);
+
+	//~ Begin IDMXProtocolInterface implementation
+	virtual IDMXProtocol* GetProtocol() const override;
+	//~ End IDMXProtocolInterface implementation
+
+private:
+	IDMXProtocol* DMXProtocol;
+};
\ No newline at end of file
diff --git a/Engine/Plugins/Experimental/DigitalMultiplex/DMXProtocol/Source/DMXProtocolArtNet/Private/DMXProtocolPortArtNet.cpp b/Engine/Plugins/Experimental/DigitalMultiplex/DMXProtocol/Source/DMXProtocolArtNet/Private/DMXProtocolPortArtNet.cpp
new file mode 100644
--- /dev/null
+++ b/Engine/Plugins/Experimental/DigitalMultiplex/DMXProtocol/Source/DMXProtocolArtNet/Private/DMXProtocolPortArtNet.cpp
@@ -0,0 +1,186 @@
+#include "DMXProtocolPortArtNet.h"
+#include "Interfaces/IDMXProtocol.h"
+#include "Interfaces/IDMXProtocolUniverse.h"
+
+#include "Json.h"
+
+//~ Artnet specific port constants
+static const uint8 HIGH_NIBBLE = 0xF0;
+static const uint8 LOW_NIBBLE = 0x0F;
+static const uint8 MAX_NET = 0x7f;
+
+FDMXProtocolPortArtNet::FDMXProtocolPortArtNet(IDMXProtocol* InDMXProtocol, TSharedPtr<IDMXProtocolDevice> InProtocolDevice, const FJsonObject& InSettings, uint8 InPortID, EDMXPortDirection InPortDirection)
+	: DMXProtocol(InDMXProtocol)
+	, ProtocolDevice(InProtocolDevice)
+	, PortID(InPortID)
+	, PortAddress(0)
+	, NetAddress(0)
+	, SubnetAddress(0)
+	, UniverseAddress(0)
+	, bIsRDMSupport(true)
+	, PortDirection(InPortDirection)
+{
+	checkf(DMXProtocol, TEXT("DMXProtocol pointer is nullptr"));
+
+	Settings = MakeShared<FJsonObject>(InSettings);
+
+	// Try to set Subnet and Universe
+	/* Bits 15 			Bits 14-8        | Bits 7-4      | Bits 3-0
+	* 0 				Net 			 | Sub-Net 		 | Universe
+	* 0				    (0b1111111 << 8) | (0b1111 << 4) | (0b1111)
+	*/
+	uint8 Net = Settings->GetNumberField(TEXT("Net"));
+	uint8 Subnet = Settings->GetNumberField(TEXT("Subnet"));
+	uint8 Universe = Settings->GetNumberField(TEXT("Universe"));
+	SetNetAddress(Net);
+	SetSubnetAddress(Subnet);
+	SetUniverse(Universe);
+
+	UniverseID = (Net << 8) | (Subnet << 4) | Universe;
+}
+
+IDMXProtocol * FDMXProtocolPortArtNet::GetProtocol() const
+{
+	return DMXProtocol;
+}
+
+uint8 FDMXProtocolPortArtNet::GetPortID() const
+{
+	return PortID;
+}
+
+uint8 FDMXProtocolPortArtNet::GetPriority() const
+{
+	return Priority;
+}
+
+EDMXPortCapability FDMXProtocolPortArtNet::GetPortCapability() const
+{
+	return DMX_PORT_CAPABILITY_FULL;
+}
+
+EDMXPortDirection FDMXProtocolPortArtNet::GetPortDirection() const
+{
+	return PortDirection;
+}
+
+bool FDMXProtocolPortArtNet::WriteDMX(const TSharedPtr<FDMXBuffer>& DMXBuffer)
+{
+	return DMXProtocol->SendDMX(UniverseID, PortID, DMXBuffer);
+}
+
+void FDMXProtocolPortArtNet::SetPriotiy(uint8 InPriority)
+{
+	Priority = InPriority;
+}
+
+bool FDMXProtocolPortArtNet::IsSupportRDM() const
+{
+	return bIsRDMSupport;
+}
+
+TWeakPtr<IDMXProtocolDevice> FDMXProtocolPortArtNet::GetCachedDevice() const
+{
+	return ProtocolDevice;
+}
+
+TWeakPtr<IDMXProtocolUniverse> FDMXProtocolPortArtNet::GetCachedUniverse() const
+{
+	return ProtocolUniverse;
+}
+
+void FDMXProtocolPortArtNet::SetUniverse(const TSharedPtr<IDMXProtocolUniverse>& InUniverse)
+{
+	ProtocolUniverse = InUniverse;
+}
+
+bool FDMXProtocolPortArtNet::ReadDMX()
+{
+	return false;
+}
+
+TSharedPtr<FJsonObject> FDMXProtocolPortArtNet::GetSettings() const
+{
+	return Settings;
+}
+
+uint16 FDMXProtocolPortArtNet::GetUniverseID() const
+{
+	return UniverseID;
+}
+
+bool FDMXProtocolPortArtNet::SetNetAddress(uint8 InNet)
+{
+	if (InNet & 0x80) 
+	{
+		UE_LOG_DMXPROTOCOL(Warning, TEXT("Art-Net net address > 127, truncating!"));
+		NetAddress = InNet & MAX_NET;
+		return true;
+	}
+	if (NetAddress == InNet)
+	{
+		return false;
+	}
+
+	NetAddress = InNet;
+	return true;
+}
+
+bool FDMXProtocolPortArtNet::SetSubnetAddress(uint8 InSubnetAddress)
+{
+	// Bits 7-4
+	// Shift for 4 bytes
+	SubnetAddress = InSubnetAddress << 4;
+	if (SubnetAddress == (PortAddress & HIGH_NIBBLE)) 
+	{
+		return false;
+	}
+
+	PortAddress = SubnetAddress | (PortAddress & LOW_NIBBLE);
+
+	return true;
+}
+
+bool FDMXProtocolPortArtNet::SetUniverse(uint8 InUniverse)
+{
+	// Bits 1-3
+	UniverseAddress = InUniverse & LOW_NIBBLE;
+	if ((PortAddress & LOW_NIBBLE) == UniverseAddress) 
+	{
+		return false;
+	}
+
+	PortAddress = ((PortAddress & HIGH_NIBBLE) | UniverseAddress);
+
+	return true;
+}
+
+uint8 FDMXProtocolPortArtNet::GetPortAddress() const
+{
+	return PortAddress;
+}
+
+uint8 FDMXProtocolPortArtNet::GetNetAddress() const
+{
+	return NetAddress;
+}
+
+uint8 FDMXProtocolPortArtNet::GetSubnetAddress() const
+{
+	return SubnetAddress;
+}
+
+uint8 FDMXProtocolPortArtNet::GetUniverseAddress() const
+{
+	return UniverseAddress;
+}
+
+const TArray<FRDMUID>& FDMXProtocolPortArtNet::GetTODUIDs() const
+{
+	return TODUIDs;
+}
+
+void FDMXProtocolPortArtNet::AddTODUID(const FRDMUID & InUID)
+{
+	TODUIDs.Add(InUID);
+}
diff --git a/Engine/Plugins/Experimental/DigitalMultiplex/DMXProtocol/Source/DMXProtocolArtNet/Private/DMXProtocolPortArtNet.h b/Engine/Plugins/Experimental/DigitalMultiplex/DMXProtocol/Source/DMXProtocolArtNet/Private/DMXProtocolPortArtNet.h
new file mode 100644
--- /dev/null
+++ b/Engine/Plugins/Experimental/DigitalMultiplex/DMXProtocol/Source/DMXProtocolArtNet/Private/DMXProtocolPortArtNet.h
@@ -0,0 +1,87 @@
+// Copyright 1998-2019 Epic Games, Inc. All Rights Reserved.
+
+#pragma once
+
+#include "CoreMinimal.h"
+#include "Interfaces/IDMXProtocolPort.h"
+#include "DMXProtocolTypes.h"
+
+class IDMXProtocol;
+
+class DMXPROTOCOLARTNET_API FDMXProtocolPortArtNet
+	: public IDMXProtocolPort
+{
+public:
+	FDMXProtocolPortArtNet(IDMXProtocol* InDMXProtocol, TSharedPtr<IDMXProtocolDevice> InProtocolDevice, const FJsonObject& InSettings, uint8 InPortID, EDMXPortDirection InPortDirection);
+
+	//~ Begin IDMXProtocolPort implementation
+	virtual IDMXProtocol* GetProtocol() const override;
+
+	virtual uint8 GetPortID() const override;
+	virtual uint8 GetPriority() const override;
+	virtual void SetPriotiy(uint8 InPriority) override;
+	virtual bool IsSupportRDM() const override;
+	virtual TWeakPtr<IDMXProtocolDevice> GetCachedDevice() const override;
+	virtual TWeakPtr<IDMXProtocolUniverse> GetCachedUniverse() const override;
+	virtual void SetUniverse(const TSharedPtr<IDMXProtocolUniverse>& InUniverse) override;
+	virtual EDMXPortCapability GetPortCapability() const override;
+	virtual EDMXPortDirection GetPortDirection() const override;
+	virtual bool WriteDMX(const TSharedPtr<FDMXBuffer>& DMXBuffer) override;
+	virtual bool ReadDMX() override;
+	virtual TSharedPtr<FJsonObject> GetSettings() const override;
+	virtual uint16 GetUniverseID() const override;
+
+	/**
+	 * The 8-bit port address, which is made up of the sub-net and universe
+	 * @return The 8-bit port address
+	 */
+	virtual uint8 GetPortAddress() const override;
+	//~ End IDMXProtocolPort implementation
+
+private:
+	//~ ArtNet specific implemntation
+	bool SetNetAddress(uint8 InNet);
+	bool SetSubnetAddress(uint8 InSubnetAddress);
+	bool SetUniverse(uint8 InUniverse);
+
+public:
+	/**
+	 * Bits 14-8 
+	 * @return The 8-bit net address
+	 */
+	uint8 GetNetAddress() const;
+
+	/**
+	 * Bits 7-4
+	 * @return The 8-bit subnet address
+	 */
+	uint8 GetSubnetAddress() const;
+
+	/**
+	 * Bits 3-0
+	 * @return The 8-bit uni address
+	 */
+	uint8 GetUniverseAddress() const;
+
+	const TArray<FRDMUID>& GetTODUIDs() const;
+
+	void AddTODUID(const FRDMUID& InUID);
+
+private:
+	IDMXProtocol* DMXProtocol;
+	TWeakPtr<IDMXProtocolDevice> ProtocolDevice;
+	TWeakPtr<IDMXProtocolUniverse> ProtocolUniverse;
+	uint8 PortID;
+	uint8 PortAddress;
+	uint8 NetAddress;
+	uint8 SubnetAddress;
+	uint8 UniverseAddress;
+	uint8 Priority;
+	uint16 UniverseID;
+	bool bIsRDMSupport;
+	EDMXPortDirection PortDirection;
+	TSharedPtr<FJsonObject> Settings;
+
+	// Array of TODs UID. This is using for RDM discovery
+	TArray<FRDMUID> TODUIDs;
+};
diff --git a/Engine/Plugins/Experimental/DigitalMultiplex/DMXProtocol/Source/DMXProtocolArtNet/Private/DMXProtocolTransportArtNet.cpp b/Engine/Plugins/Experimental/DigitalMultiplex/DMXProtocol/Source/DMXProtocolArtNet/Private/DMXProtocolTransportArtNet.cpp
new file mode 100644
--- /dev/null
+++ b/Engine/Plugins/Experimental/DigitalMultiplex/DMXProtocol/Source/DMXProtocolArtNet/Private/DMXProtocolTransportArtNet.cpp
@@ -0,0 +1,212 @@
+// Copyright 1998-2019 Epic Games, Inc. All Rights Reserved.
+
+#include "DMXProtocolTransportArtNet.h"
+#include "DMXProtocolArtNet.h"
+
+#include "Common/UdpSocketSender.h"
+#include "HAL/Event.h"
+#include "HAL/Runnable.h"
+#include "HAL/RunnableThread.h"
+#include "IMessageAttachment.h"
+#include "Serialization/ArrayReader.h"
+#include "Serialization/ArrayWriter.h"
+#include "Sockets.h"
+
+static FTimespan CalculateWaitTime()
+{
+	return FTimespan::FromMilliseconds(10);
+}
+
+FDMXProtocolSenderArtNet::FDMXProtocolSenderArtNet(FSocket& InSocket, FDMXProtocolArtNet* InProtocol)
+	: LastSentPackage(-1)
+	, bRequestingExit(false)
+	, BroadcastSocket(&InSocket)
+	, Protocol(InProtocol)
+{
+	WorkEvent = MakeShareable(FPlatformProcess::GetSynchEventFromPool(), [](FEvent* EventToDelete)
+	{
+		FPlatformProcess::ReturnSynchEventToPool(EventToDelete);
+	});
+
+	Thread = FRunnableThread::Create(this, TEXT("FDMXProtocolSenderArtNet"), 128 * 1024, TPri_BelowNormal, FPlatformAffinity::GetPoolThreadMask());
+}
+
+FDMXProtocolSenderArtNet::~FDMXProtocolSenderArtNet()
+{
+	// Shut Down worker thread
+	if (Thread != nullptr)
+	{
+		Thread->Kill(true);
+		delete Thread;
+	}
+}
+
+bool FDMXProtocolSenderArtNet::Init()
+{
+	return true;
+}
+
+uint32 FDMXProtocolSenderArtNet::Run()
+{
+	while (StopTaskCounter.GetValue() == 0)
+	{
+		if (WorkEvent->Wait(CalculateWaitTime()))
+		{
+			ConsumeOutboundPackages();
+		}
+	}
+
+	return 0;
+}
+
+void FDMXProtocolSenderArtNet::Stop()
+{
+	StopTaskCounter.Increment();
+
+	if (WorkEvent)
+	{
+		WorkEvent->Trigger();
+		Thread->WaitForCompletion();
+	}
+}
+
+void FDMXProtocolSenderArtNet::Exit()
+{
+}
+
+void FDMXProtocolSenderArtNet::Tick()
+{
+	ConsumeOutboundPackages();
+}
+
+FSingleThreadRunnable* FDMXProtocolSenderArtNet::GetSingleThreadInterface()
+{
+	return this;
+}
+
+bool FDMXProtocolSenderArtNet::EnqueueOutboundPackage(FDMXPacketPtr Packet)
+{
+	if (StopTaskCounter.GetValue())
+	{
+		return false;
+	}
+
+	if (!OutboundPackages.Enqueue(Packet))
+	{
+		return false;
+	}
+
+	WorkEvent->Trigger();
+
+
+	return true;
+}
+
+void FDMXProtocolSenderArtNet::ConsumeOutboundPackages()
+{
+	FScopeLock ScopeLock(&CriticalSection);
+
+	FDMXPacketPtr Packet;
+	while (OutboundPackages.Dequeue(Packet))
+	{
+		++LastSentPackage;
+		int32 BytesSent = 0;
+		bool bIsSent = BroadcastSocket->SendTo(Packet->Data.GetData(), Packet->Data.Num(), BytesSent, *Protocol->GetBroadcastAddr());
+		if (!bIsSent)
+		{
+			UE_LOG_DMXPROTOCOL(Error, TEXT("Error sending %d"), bIsSent);
+		}
+	}
+}
+
+FDMXProtocolReceiverArtNet::FDMXProtocolReceiverArtNet(FSocket& InSocket, FDMXProtocolArtNet* InProtocol, const FTimespan& InWaitTime)
+	: Socket(&InSocket)
+	, Stopping(false)
+	, Thread(nullptr)
+	, WaitTime(InWaitTime)
+{
+	check(Socket != nullptr);
+	check(Socket->GetSocketType() == SOCKTYPE_Datagram);
+
+	SocketSubsystem = ISocketSubsystem::Get(PLATFORM_SOCKETSUBSYSTEM);
+
+	Thread = FRunnableThread::Create(this, *ThreadName, 128 * 1024, TPri_Normal, FPlatformAffinity::GetPoolThreadMask());
+}
+
+FDMXProtocolReceiverArtNet::~FDMXProtocolReceiverArtNet()
+{
+	if (Thread != nullptr)
+	{
+		Thread->Kill(true);
+		delete Thread;
+	}
+}
+
+FOnDMXDataReceived& FDMXProtocolReceiverArtNet::OnDataReceived()
+{
+	return DMXDataReceiveDelegate;
+}
+
+FRunnableThread* FDMXProtocolReceiverArtNet::GetThread() const
+{
+	return Thread;
+}
+
+bool FDMXProtocolReceiverArtNet::Init()
+{
+	return true;
+}
+
+uint32 FDMXProtocolReceiverArtNet::Run()
+{
+	while (!Stopping)
+	{
+		Update(WaitTime);
+	}
+
+	return 0;
+}
+
+void FDMXProtocolReceiverArtNet::Stop()
+{
+	Stopping = true;
+}
+
+void FDMXProtocolReceiverArtNet::Exit()
+{
+}
+
+void FDMXProtocolReceiverArtNet::Tick()
+{
+	Update(FTimespan::Zero());
+}
+
+FSingleThreadRunnable* FDMXProtocolReceiverArtNet::GetSingleThreadInterface()
+{
+	return this;
+}
+
+void FDMXProtocolReceiverArtNet::Update(const FTimespan& SocketWaitTime)
+{
+	if (!Socket->Wait(ESocketWaitConditions::WaitForRead, SocketWaitTime))
+	{
+		return;
+	}
+
+	TSharedRef<FInternetAddr> Sender = SocketSubsystem->CreateInternetAddr();
+	uint32 Size;
+
+	while (Socket->HasPendingData(Size))
+	{
+		FArrayReaderPtr Reader = MakeShareable(new FArrayReader(true));
+		Reader->SetNumUninitialized(FMath::Min(Size, 65507u));
+
+		int32 Read = 0;
+
+		if (Socket->RecvFrom(Reader->GetData(), Reader->Num(), Read, *Sender))
+		{
+			Reader->RemoveAt(Read, Reader->Num() - Read, false);
+			DMXDataReceiveDelegate.ExecuteIfBound(Reader);
+		}
+	}
+}
diff --git a/Engine/Plugins/Experimental/DigitalMultiplex/DMXProtocol/Source/DMXProtocolArtNet/Private/DMXProtocolTransportArtNet.h b/Engine/Plugins/Experimental/DigitalMultiplex/DMXProtocol/Source/DMXProtocolArtNet/Private/DMXProtocolTransportArtNet.h
new file mode 100644
--- /dev/null
+++ b/Engine/Plugins/Experimental/DigitalMultiplex/DMXProtocol/Source/DMXProtocolArtNet/Private/DMXProtocolTransportArtNet.h
@@ -0,0 +1,122 @@
+// Copyright 1998-2019 Epic Games, Inc. All Rights Reserved.
+
+#pragma once
+
+#include "CoreMinimal.h"
+#include "Containers/Queue.h"
+#include "Interfaces/IDMXProtocolTransport.h"
+#include "DMXProtocolTypes.h"
+
+
+class FSocket;
+class FDMXProtocolArtNet;
+
+
+class DMXPROTOCOLARTNET_API FDMXProtocolSenderArtNet
+	: public IDMXProtocolSender
+{
+
+public:
+	FDMXProtocolSenderArtNet(FSocket& InSocket, FDMXProtocolArtNet* InProtocol);
+	virtual ~FDMXProtocolSenderArtNet();
+
+	//~ Begin FRunnable implementation
+	virtual bool Init() override;
+	virtual uint32 Run() override;
+	virtual void Stop() override;
+	virtual void Exit() override;
+	//~ End FRunnable implementation
+
+	//~ Begin FSingleThreadRunnable implementation
+	virtual void Tick() override;
+	virtual class FSingleThreadRunnable* GetSingleThreadInterface() override;
+	//~ End FSingleThreadRunnable implementation
+
+	//~ Begin IDMXProtocolSender implementation
+	virtual bool EnqueueOutboundPackage(FDMXPacketPtr Packet) override;
+	//~ End IDMXProtocolSender implementation
+
+public:
+	/** Consumes all outbound packages. */
+	void ConsumeOutboundPackages();
+
+
+private:
+	/** Holds the queue of outbound packages. */
+	TQueue<FDMXPacketPtr, EQueueMode::Mpsc> OutboundPackages;
+
+	/** Holds the last sent message number. */
+	int32 LastSentPackage;
+
+	mutable FCriticalSection CriticalSection;
+
+	FThreadSafeCounter StopTaskCounter;
+
+	/** Holds the thread object. */
+	FRunnableThread* Thread;
+
+	/** Holds an event signaling that inbound messages need to be processed. */
+	TSharedPtr<FEvent, ESPMode::ThreadSafe> WorkEvent;
+
+	FThreadSafeBool bRequestingExit;
+
+	/** Holds the network socket used to sender packages. */
+	FSocket* BroadcastSocket;
+
+	FDMXProtocolArtNet* Protocol;
+};
+
+
+class FDMXProtocolReceiverArtNet
+	: public IDMXProtocolReceiver
+{
+public:
+	FDMXProtocolReceiverArtNet(FSocket& InSocket, FDMXProtocolArtNet* InProtocol, const FTimespan& InWaitTime);
+	virtual ~FDMXProtocolReceiverArtNet();
+
+public:
+	//~ Begin FRunnable implementation
+	virtual bool Init() override;
+	virtual uint32 Run() override;
+	virtual void Stop() override;
+	virtual void Exit() override;
+	//~ End FRunnable implementation
+
+	//~ Begin FSingleThreadRunnable implementation
+	virtual void Tick() override;
+	virtual class FSingleThreadRunnable* GetSingleThreadInterface() override;
+	//~ End FSingleThreadRunnable implementation
+
+	//~ Begin IDMXProtocolReceiver implementation
+	virtual FOnDMXDataReceived& OnDataReceived() override;
+	virtual FRunnableThread* GetThread() const override;
+	//~ End IDMXProtocolReceiver implementation
+
+protected:
+	void Update(const FTimespan& SocketWaitTime);
+
+private:
+
+	/** The network socket. */
+	FSocket* Socket;
+
+	/** Pointer to the socket sub-system. */
+	ISocketSubsystem* SocketSubsystem;
+
+	/** Flag indicating that the thread is stopping. */
+	bool Stopping;
+
+	/** The thread object. */
+	FRunnableThread* Thread;
+
+	/** The receiver thread's name. */
+	FString ThreadName;
+
+	/** The amount of time to wait for inbound packets. */
+	FTimespan WaitTime;
+
+private:
+
+	/** Holds the data received delegate. */
+	FOnDMXDataReceived DMXDataReceiveDelegate;
+};
\ No newline at end of file
diff --git a/Engine/Plugins/Experimental/DigitalMultiplex/DMXProtocol/Source/DMXProtocolArtNet/Private/DMXProtocolUniverseArtNet.cpp b/Engine/Plugins/Experimental/DigitalMultiplex/DMXProtocol/Source/DMXProtocolArtNet/Private/DMXProtocolUniverseArtNet.cpp
new file mode 100644
--- /dev/null
+++ b/Engine/Plugins/Experimental/DigitalMultiplex/DMXProtocol/Source/DMXProtocolArtNet/Private/DMXProtocolUniverseArtNet.cpp
@@ -0,0 +1,58 @@
+#include "DMXProtocolUniverseArtNet.h"
+#include "Json.h"
+
+#include "Interfaces/IDMXProtocol.h"
+#include "Interfaces/IDMXProtocolPort.h"
+#include "Interfaces/IDMXProtocolDevice.h"
+#include "DMXProtocolConstants.h"
+#include "DMXProtocolTypes.h"
+
+FDMXProtocolUniverseArtNet::FDMXProtocolUniverseArtNet(IDMXProtocol* InDMXProtocol, TSharedPtr<IDMXProtocolPort> InPort)
+	: DMXProtocol(InDMXProtocol)
+	, Port(InPort)
+	, Priority(0)
+	, UniverseID(InPort->GetUniverseID())
+{
+	checkf(DMXProtocol, TEXT("DMXProtocol pointer is nullptr"));
+	checkf(Port.IsValid(), TEXT("DMXProtocol port is not valid"));
+
+	InputDMXBuffer = MakeShared<FDMXBuffer>();
+	OutputDMXBuffer = MakeShared<FDMXBuffer>();
+}
+
+
+IDMXProtocol * FDMXProtocolUniverseArtNet::GetProtocol() const
+{
+	return DMXProtocol;
+}
+
+TSharedPtr<FDMXBuffer> FDMXProtocolUniverseArtNet::GetOutputDMXBuffer() const
+{
+	return OutputDMXBuffer;
+}
+
+TSharedPtr<FDMXBuffer> FDMXProtocolUniverseArtNet::GetInputDMXBuffer() const
+{
+	return InputDMXBuffer;
+}
+
+bool FDMXProtocolUniverseArtNet::SetDMXFragment(const IDMXFragmentMap & DMXFragment)
+{
+	OutputDMXBuffer->SetDMXFragment(DMXFragment);
+	return false;
+}
+
+uint8 FDMXProtocolUniverseArtNet::GetPriority() const
+{
+	return Priority;
+}
+
+uint16 FDMXProtocolUniverseArtNet::GetUniverseID() const
+{
+	return UniverseID;
+}
+
+TWeakPtr<IDMXProtocolPort> FDMXProtocolUniverseArtNet::GetCachedUniversePort() const
+{
+	return Port;
+}
diff --git a/Engine/Plugins/Experimental/DigitalMultiplex/DMXProtocol/Source/DMXProtocolArtNet/Private/DMXProtocolUniverseArtNet.h b/Engine/Plugins/Experimental/DigitalMultiplex/DMXProtocol/Source/DMXProtocolArtNet/Private/DMXProtocolUniverseArtNet.h
new file mode 100644
--- /dev/null
+++ b/Engine/Plugins/Experimental/DigitalMultiplex/DMXProtocol/Source/DMXProtocolArtNet/Private/DMXProtocolUniverseArtNet.h
@@ -0,0 +1,31 @@
+// Copyright 1998-2019 Epic Games, Inc. All Rights Reserved.
+
+#pragma once
+
+#include "CoreMinimal.h"
+#include "Interfaces/IDMXProtocolUniverse.h"
+
+class DMXPROTOCOLARTNET_API FDMXProtocolUniverseArtNet
+	: public IDMXProtocolUniverse
+{
+public:
+	FDMXProtocolUniverseArtNet(IDMXProtocol* InDMXProtocol, TSharedPtr<IDMXProtocolPort> InPort);
+
+	//~ Begin IDMXProtocolDevice implementation
+	virtual IDMXProtocol* GetProtocol() const override;
+	virtual TSharedPtr<FDMXBuffer> GetOutputDMXBuffer() const override;
+	virtual TSharedPtr<FDMXBuffer> GetInputDMXBuffer() const override;
+	virtual bool SetDMXFragment(const IDMXFragmentMap& DMXFragment) override;
+	virtual uint8 GetPriority() const override;
+	virtual uint16 GetUniverseID() const override;
+	virtual TWeakPtr<IDMXProtocolPort> GetCachedUniversePort() const override;
+	//~ End IDMXProtocolDevice implementation
+
+private:
+	IDMXProtocol* DMXProtocol;
+	TWeakPtr<IDMXProtocolPort> Port;
+	TSharedPtr<FDMXBuffer> OutputDMXBuffer;
+	TSharedPtr<FDMXBuffer> InputDMXBuffer;
+	uint8 Priority;
+	uint16 UniverseID;
+};
\ No newline at end of file
diff --git a/Engine/Plugins/Experimental/DigitalMultiplex/DMXProtocol/Source/DMXProtocolArtNet/Private/Packets/DMXProtocolArtNetPackets.cpp b/Engine/Plugins/Experimental/DigitalMultiplex/DMXProtocol/Source/DMXProtocolArtNet/Private/Packets/DMXProtocolArtNetPackets.cpp
new file mode 100644
--- /dev/null
+++ b/Engine/Plugins/Experimental/DigitalMultiplex/DMXProtocol/Source/DMXProtocolArtNet/Private/Packets/DMXProtocolArtNetPackets.cpp
@@ -0,0 +1,209 @@
+#include "Packets/DMXProtocolArtNetPackets.h"
+#include "DMXProtocolTypes.h"
+#include "Serialization/DMXMemoryWriter.h"
+#include "DMXProtocolMacros.h"
+
+REGISTER_DMX_ARCHIVE(FDMXProtocolArtNetDMXPacket);
+REGISTER_DMX_ARCHIVE(FDMXProtocolArtNetPollPacket);
+REGISTER_DMX_ARCHIVE(FArtNetPacketReply);
+REGISTER_DMX_ARCHIVE(FDMXProtocolArtNetTodRequest);
+REGISTER_DMX_ARCHIVE(FDMXProtocolArtNetTodData);
+REGISTER_DMX_ARCHIVE(FDMXProtocolArtNetTodControl);
+REGISTER_DMX_ARCHIVE(FDMXProtocolArtNetRDM);
+
+TSharedPtr<FBufferArchive> FDMXProtocolArtNetDMXPacket::Pack()
+{
+	TSharedPtr<FDMXMemoryWriter> Writer = MakeShared<FDMXMemoryWriter>();
+	*Writer << *this;
+
+	return Writer;
+}
+
+void FDMXProtocolArtNetDMXPacket::Serialize(FArchive & Ar)
+{
+	Ar.Serialize((void*)ID, ARTNET_STRING_SIZE);
+	Ar << OpCode;
+	Ar << VerH;
+	Ar << Ver;
+	Ar << Sequence;
+	Ar << Physical;
+	Ar << Universe;
+	Ar << LengthHi;
+	Ar << Length;
+	Ar.Serialize((void*)Data, ARTNET_DMX_LENGTH);
+}
+
+
+TSharedPtr<FBufferArchive> FDMXProtocolArtNetPollPacket::Pack()
+{
+	TSharedPtr<FDMXMemoryWriter> Writer = MakeShared<FDMXMemoryWriter>();
+	*Writer << *this;
+
+	return Writer;
+}
+
+void FDMXProtocolArtNetPollPacket::Serialize(FArchive & Ar)
+{
+	Ar.Serialize((void*)ID, ARTNET_STRING_SIZE);
+	Ar << OpCode;
+	Ar << VerH;
+	Ar << Ver;
+	Ar << TalkToMe;
+	Ar << Priority;
+}
+
+void FArtNetPacketReply::Serialize(FArchive & Ar)
+{
+	Ar.Serialize((void*)ID, ARTNET_STRING_SIZE);
+	Ar << OpCode;
+	Ar.Serialize((void*)Ip, sizeof(Ip));
+	Ar << Port;
+	Ar << Version;
+	Ar << NetAddress;
+	Ar << SubnetAddress;
+	Ar << Oem;
+	Ar << Ubea;
+	Ar << Status1;
+	Ar << EstaId;
+	Ar.Serialize((void*)ShortName, ARTNET_SHORT_NAME_LENGTH);
+	Ar.Serialize((void*)LongName, ARTNET_LONG_NAME_LENGTH);
+	Ar.Serialize((void*)NodeReport, ARTNET_REPORT_LENGTH);
+	Ar.Serialize((void*)NumberPorts, sizeof(NumberPorts));
+	Ar.Serialize((void*)PortTypes, ARTNET_MAX_PORTS);
+	Ar.Serialize((void*)GoodInput, ARTNET_MAX_PORTS);
+	Ar.Serialize((void*)GoodOutput, ARTNET_MAX_PORTS);
+	Ar.Serialize((void*)SwIn, ARTNET_MAX_PORTS);
+	Ar.Serialize((void*)SwOut, ARTNET_MAX_PORTS);
+	Ar << SwVideo;
+	Ar << SwMacro;
+	Ar << SwRemote;
+	Ar << Spare1;
+	Ar << Spare2;
+	Ar << Spare3;
+	Ar << Style;
+	Ar.Serialize((void*)Mac, sizeof(Mac));
+	Ar.Serialize((void*)BindIp, sizeof(BindIp));
+	Ar << BindIndex;
+	Ar << Status2;
+	Ar.Serialize((void*)Filler, sizeof(Filler));
+}
+
+TSharedPtr<FBufferArchive> FDMXProtocolArtNetTodRequest::Pack()
+{
+	TSharedPtr<FDMXMemoryWriter> Writer = MakeShared<FDMXMemoryWriter>();
+	*Writer << *this;
+
+	return Writer;
+}
+
+void FDMXProtocolArtNetTodRequest::Serialize(FArchive & Ar)
+{
+	Ar.Serialize((void*)ID, ARTNET_STRING_SIZE);
+	Ar << OpCode;
+	Ar << VerH;
+	Ar << Ver;
+	Ar << Filler1;
+	Ar << Filler2;
+	Ar << Spare1;
+	Ar << Spare2;
+	Ar << Spare3;
+	Ar << Spare4;
+	Ar << Spare5;
+	Ar << Spare6;
+	Ar << Spare7;
+	Ar << Net;
+	Ar << Command;
+	Ar << AdCount;
+	Ar.Serialize((void*)Address, ARTNET_MAX_RDM_ADCOUNT);
+}
+
+TSharedPtr<FBufferArchive> FDMXProtocolArtNetTodData::Pack()
+{
+	TSharedPtr<FDMXMemoryWriter> Writer = MakeShared<FDMXMemoryWriter>();
+	*Writer << *this;
+
+	return Writer;
+}
+
+void FDMXProtocolArtNetTodData::Serialize(FArchive & Ar)
+{
+	Ar.Serialize((void*)ID, ARTNET_STRING_SIZE);
+	Ar << OpCode;
+	Ar << VerH;
+	Ar << Ver;
+	Ar << RDMVer;
+	Ar << Port;
+	Ar << Spare1;
+	Ar << Spare2;
+	Ar << Spare3;
+	Ar << Spare4;
+	Ar << Spare5;
+	Ar << Spare6;
+	Ar << Spare7;
+	Ar << Net;
+	Ar << CmdRes;
+	Ar << Address;
+	Ar << UidTotalHi;
+	Ar << UidTotal;
+	Ar << BlockCount;
+	Ar << UidCount;
+	Ar.Serialize((void*)Tod, sizeof(Tod));
+}
+
+TSharedPtr<FBufferArchive> FDMXProtocolArtNetTodControl::Pack()
+{
+	TSharedPtr<FDMXMemoryWriter> Writer = MakeShared<FDMXMemoryWriter>();
+	*Writer << *this;
+
+	return Writer;
+}
+
+void FDMXProtocolArtNetTodControl::Serialize(FArchive & Ar)
+{
+	Ar.Serialize((void*)ID, ARTNET_STRING_SIZE);
+	Ar << OpCode;
+	Ar << VerH;
+	Ar << Ver;
+	Ar << Filler1;
+	Ar << Filler2;
+	Ar << Spare1;
+	Ar << Spare2;
+	Ar << Spare3;
+	Ar << Spare4;
+	Ar << Spare5;
+	Ar << Spare6;
+	Ar << Spare7;
+	Ar << Net;
+	Ar << Cmd;
+	Ar << Address;
+}
+
+TSharedPtr<FBufferArchive> FDMXProtocolArtNetRDM::Pack()
+{
+	TSharedPtr<FDMXMemoryWriter> Writer = MakeShared<FDMXMemoryWriter>();
+	*Writer << *this;
+
+	return Writer;
+}
+
+void FDMXProtocolArtNetRDM::Serialize(FArchive & Ar)
+{
+	Ar.Serialize((void*)ID, ARTNET_STRING_SIZE);
+	Ar << OpCode;
+	Ar << VerH;
+	Ar << Ver;
+	Ar << RdmVer;
+	Ar << Filler2;
+	Ar << Spare1;
+	Ar << Spare2;
+	Ar << Spare3;
+	Ar << Spare4;
+	Ar << Spare5;
+	Ar << Spare6;
+	Ar << Spare7;
+	Ar << Net;
+	Ar << Cmd;
+	Ar << Address;
+	Ar.Serialize((void*)Data, ARTNET_MAX_RDM_DATA);
+}
+
diff --git a/Engine/Plugins/Experimental/DigitalMultiplex/DMXProtocol/Source/DMXProtocolArtNet/Private/Packets/DMXProtocolArtNetPackets.h b/Engine/Plugins/Experimental/DigitalMultiplex/DMXProtocol/Source/DMXProtocolArtNet/Private/Packets/DMXProtocolArtNetPackets.h
new file mode 100644
--- /dev/null
+++ b/Engine/Plugins/Experimental/DigitalMultiplex/DMXProtocol/Source/DMXProtocolArtNet/Private/Packets/DMXProtocolArtNetPackets.h
@@ -0,0 +1,236 @@
+// Copyright 1998-2019 Epic Games, Inc. All Rights Reserved.
+
+#pragma once
+
+#include "CoreMinimal.h"
+#include "Interfaces/IDMXProtocolPacket.h"
+#include "DMXProtocolArtNetConstants.h"
+
+struct FDMXProtocolArtNetPacket
+	: public IDMXProtocolPacket
+{
+public:
+	FDMXProtocolArtNetPacket()
+	{
+		FMemory::Memcpy(ID, DMXArtNet::ARTNET_STRING, ARTNET_STRING_SIZE);
+	}
+
+public:
+	uint8 ID[8] = { 0 };
+};
+
+struct FDMXProtocolArtNetDMXPacket
+	: public FDMXProtocolArtNetPacket
+{
+public:
+	//~ Begin IDMXProtocolPacket implementation
+	virtual TSharedPtr<FBufferArchive> Pack() override;
+	//~ End IDMXProtocolPort implementation
+
+	void Serialize(FArchive& Ar);
+	friend FArchive& operator<<(FArchive& Ar, FDMXProtocolArtNetDMXPacket& Packet);
+
+public:
+	uint16 OpCode = ARTNET_DMX;
+
+	uint8 VerH = 0;
+	uint8 Ver = ARTNET_VERSION;
+	uint8 Sequence = 0;
+	uint8 Physical = 0;
+	uint16 Universe = 0;
+	uint8 LengthHi = DMX_SHORT_GET_HIGH_BIT(ARTNET_DMX_LENGTH);
+	uint8 Length = DMX_SHORT_GET_LOW_BYTE(ARTNET_DMX_LENGTH);
+	uint8 Data[ARTNET_DMX_LENGTH] = { 0 };
+};
+
+struct FDMXProtocolArtNetPollPacket
+	: public FDMXProtocolArtNetPacket
+{
+public:
+	//~ Begin IDMXProtocolPacket implementation
+	virtual TSharedPtr<FBufferArchive> Pack() override;
+	//~ End IDMXProtocolPort implementation
+
+	void Serialize(FArchive& Ar);
+	friend FArchive& operator<<(FArchive& Ar, FDMXProtocolArtNetPollPacket& Packet);
+	
+public:
+	uint16 OpCode = ARTNET_POLL;
+	uint8 VerH = 0;
+	uint8 Ver = ARTNET_VERSION;
+	uint8 TalkToMe = ARTNET_TTM_DEFAULT;
+	uint8 Priority = 0;
+};
+
+struct FArtNetPacketReply
+	: public FDMXProtocolArtNetPacket
+{
+public:
+	void Serialize(FArchive& Ar);
+	friend FArchive& operator<<(FArchive& Ar, FArtNetPacketReply& Packet);
+
+public:
+	uint16 OpCode = 0;
+
+	uint8 Ip[4] = { 0 };
+	uint16 Port = 0;
+	uint16 Version = 0;
+	uint8 NetAddress = 0;
+	uint8 SubnetAddress = 0;
+	uint16 Oem = 0;
+	uint8 Ubea = 0;
+	uint8 Status1 = 0;
+	uint16 EstaId = 0;
+	uint8 ShortName[ARTNET_SHORT_NAME_LENGTH] = { 0 };
+	uint8 LongName[ARTNET_LONG_NAME_LENGTH] = { 0 };
+	uint8 NodeReport[ARTNET_REPORT_LENGTH] = { 0 };
+	uint8 NumberPorts[2] = { 0 };
+	uint8 PortTypes[ARTNET_MAX_PORTS] = { 0 };
+	uint8 GoodInput[ARTNET_MAX_PORTS] = { 0 };
+	uint8 GoodOutput[ARTNET_MAX_PORTS] = { 0 };
+	uint8 SwIn[ARTNET_MAX_PORTS] = { 0 };
+	uint8 SwOut[ARTNET_MAX_PORTS] = { 0 };
+	uint8 SwVideo = 0;
+	uint8 SwMacro = 0;
+	uint8 SwRemote = 0;
+	uint8 Spare1 = 0;
+	uint8 Spare2 = 0;
+	uint8 Spare3 = 0;
+	uint8 Style = 0;
+	uint8 Mac[6] = { 0 };
+	uint8 BindIp[4] = { 0 };
+	uint8 BindIndex = 0;
+	uint8 Status2 = 0;
+	uint8 Filler[26] = { 0 };
+};
+
+struct FDMXProtocolArtNetTodRequest
+	: public FDMXProtocolArtNetPacket
+{
+public:
+	//~ Begin IDMXProtocolPacket implementation
+	virtual TSharedPtr<FBufferArchive> Pack() override;
+	//~ End IDMXProtocolPort implementation
+
+	void Serialize(FArchive& Ar);
+	friend FArchive& operator<<(FArchive& Ar, FDMXProtocolArtNetTodRequest& Packet);
+
+public:
+	uint16 OpCode = ARTNET_TODREQUEST;
+
+	uint8 VerH = 0;
+	uint8 Ver = ARTNET_VERSION;
+	uint8 Filler1 = 0;
+	uint8 Filler2 = 0;
+	uint8 Spare1 = 0;
+	uint8 Spare2 = 0;
+	uint8 Spare3 = 0;
+	uint8 Spare4 = 0;
+	uint8 Spare5 = 0;
+	uint8 Spare6 = 0;
+	uint8 Spare7 = 0;
+	uint8 Net = 0;
+	uint8 Command = ARTNET_TOD_FULL;
+	uint8 AdCount = 0;
+	uint8 Address[ARTNET_MAX_RDM_ADCOUNT] = { 0 };
+};
+
+struct FDMXProtocolArtNetTodData
+	: public FDMXProtocolArtNetPacket
+{
+public:
+	//~ Begin IDMXProtocolPacket implementation
+	virtual TSharedPtr<FBufferArchive> Pack() override;
+	//~ End IDMXProtocolPort implementation
+
+	void Serialize(FArchive& Ar);
+	friend FArchive& operator<<(FArchive& Ar, FDMXProtocolArtNetTodData& Packet);
+
+public:
+	uint16 OpCode = ARTNET_TODDATA;
+
+	uint8 VerH = 0;
+	uint8 Ver = ARTNET_VERSION;
+	uint8 RDMVer = ARTNET_RDM_VERSION;
+	uint8 Port = 0;
+	uint8 Spare1 = 0;
+	uint8 Spare2 = 0;
+	uint8 Spare3 = 0;
+	uint8 Spare4 = 0;
+	uint8 Spare5 = 0;
+	uint8 Spare6 = 0;
+	uint8 Spare7 = 0;
+	uint8 Net = 0;
+	uint8 CmdRes = ARTNET_TOD_FULL;
+	uint8 Address = 0;
+	uint8 UidTotalHi = 0;
+	uint8 UidTotal = 0;
+	uint8 BlockCount = 0;
+	uint8 UidCount = 0;
+	uint8 Tod[ARTNET_MAX_UID_COUNT][ARTNET_RDM_UID_WIDTH] = { { 0 } };
+};
+
+/**
+ * The ArtTodControl packet is used to send RDM control parameters over Art-Net. The
+ * response is ArtTodData.
+ */
+struct FDMXProtocolArtNetTodControl
+	: public FDMXProtocolArtNetPacket
+{
+public:
+	//~ Begin IDMXProtocolPacket implementation
+	virtual TSharedPtr<FBufferArchive> Pack() override;
+	//~ End IDMXProtocolPort implementation
+
+	void Serialize(FArchive& Ar);
+	friend FArchive& operator<<(FArchive& Ar, FDMXProtocolArtNetTodControl& Packet);
+
+public:
+	uint16 OpCode = ARTNET_TODCONTROL;
+
+	uint8 VerH = 0;
+	uint8 Ver = ARTNET_VERSION;
+	uint8 Filler1 = 0;
+	uint8 Filler2 = 0;
+	uint8 Spare1 = 0;
+	uint8 Spare2 = 0;
+	uint8 Spare3 = 0;
+	uint8 Spare4 = 0;
+	uint8 Spare5 = 0;
+	uint8 Spare6 = 0;
+	uint8 Spare7 = 0;
+	uint8 Net = 0;
+	uint8 Cmd = 0;
+	uint8 Address = 0;
+};
+
+struct FDMXProtocolArtNetRDM
+	: public FDMXProtocolArtNetPacket
+{
+public:
+	//~ Begin IDMXProtocolPacket implementation
+	virtual TSharedPtr<FBufferArchive> Pack() override;
+	//~ End IDMXProtocolPort implementation
+
+	void Serialize(FArchive& Ar);
+	friend FArchive& operator<<(FArchive& Ar, FDMXProtocolArtNetRDM& Packet);
+
+public:
+	uint16 OpCode = ARTNET_RDM;
+
+	uint8 VerH = 0;
+	uint8 Ver = ARTNET_VERSION;
+	uint8 RdmVer = 0;
+	uint8 Filler2 = 0;
+	uint8 Spare1 = 0;
+	uint8 Spare2 = 0;
+	uint8 Spare3 = 0;
+	uint8 Spare4 = 0;
+	uint8 Spare5 = 0;
+	uint8 Spare6 = 0;
+	uint8 Spare7 = 0;
+	uint8 Net = 0;
+	uint8 Cmd = ARTNET_RDM_PROCESS_PACKET;
+	uint8 Address = 0;
+	uint8 Data[ARTNET_MAX_RDM_DATA] = { 0 };
+};
\ No newline at end of file
diff --git a/Engine/Plugins/Experimental/DigitalMultiplex/DMXProtocol/Source/DMXProtocolArtNet/Private/Tests/DMXPrtotocolArtNetFlowTest.cpp b/Engine/Plugins/Experimental/DigitalMultiplex/DMXProtocol/Source/DMXProtocolArtNet/Private/Tests/DMXPrtotocolArtNetFlowTest.cpp
new file mode 100644
--- /dev/null
+++ b/Engine/Plugins/Experimental/DigitalMultiplex/DMXProtocol/Source/DMXProtocolArtNet/Private/Tests/DMXPrtotocolArtNetFlowTest.cpp
@@ -0,0 +1,208 @@
+// Copyright 1998-2019 Epic Games, Inc. All Rights Reserved.
+
+#include "CoreMinimal.h"
+#include "Engine.h"
+#include "DMXProtocolArtNetModule.h"
+
+#include "DMXProtocolArtNet.h"
+#include "DMXProtocolPortArtNet.h"
+
+#include "Managers/DMXProtocolDeviceManager.h"
+#include "Managers/DMXProtocolInterfaceManager.h"
+#include "Managers/DMXProtocolPortManager.h"
+#include "Managers/DMXProtocolUniverseManager.h"
+
+#include "DMXProtocolInterfaceArtNet.h"
+#include "DMXProtocolDeviceArtNet.h"
+#include "DMXProtocolUniverseArtNet.h"
+#include "DMXProtocolPortArtNet.h"
+#include "DMXProtocolArtNetConstants.h"
+
+struct FDMXProtocolLauncher
+{
+	FDMXProtocolLauncher(uint8 InDeviceID, uint8 InPortID)
+		: DeviceID(InDeviceID)
+		, PortID(InPortID)
+		, InterfaceID(0)
+	{
+		// Init Protocol and Managers
+		DMXProtocol = IDMXProtocol::Get<FDMXProtocolArtNet>(FDMXProtocolArtNetModule::NAME_Artnet);
+
+		InterfaceArtNet = MakeShared<FDMXProtocolInterfaceArtNet>(DMXProtocol);
+		DMXProtocol->GetInterfaceManager()->AddInterface(InterfaceID, InterfaceArtNet);
+
+		FJsonObject DeviceSettings;
+		ProtocolDevice = MakeShared<FDMXProtocolDeviceArtNet>(DMXProtocol, InterfaceArtNet, DeviceSettings, DeviceID);
+		DMXProtocol->GetDeviceManager()->AddDevice(ProtocolDevice);
+
+		FJsonObject OutputPortSettings;
+		OutputPortSettings.SetNumberField(TEXT("Net"), 0);
+		OutputPortSettings.SetNumberField(TEXT("Subnet"), 1);
+		OutputPortSettings.SetNumberField(TEXT("Universe"), 0);
+		ProtocolOutputPort = MakeShared<FDMXProtocolPortArtNet>(DMXProtocol, ProtocolDevice, OutputPortSettings, PortID, DMX_PORT_OUTPUT);
+		DMXProtocol->GetPortManager()->AddOutputPort(ProtocolDevice, ProtocolOutputPort);
+
+		ProtocolUniverse = MakeShared<FDMXProtocolUniverseArtNet>(DMXProtocol, ProtocolOutputPort);
+		DMXProtocol->GetUniverseManager()->AddUniverse(ProtocolOutputPort, ProtocolUniverse);
+		ProtocolOutputPort->SetUniverse(ProtocolUniverse);
+	}
+
+	~FDMXProtocolLauncher()
+	{
+		DMXProtocol->GetInterfaceManager()->RemoveInterface(InterfaceID);
+		DMXProtocol->GetDeviceManager()->RemoveDevice(ProtocolDevice);
+		DMXProtocol->GetPortManager()->RemoveOutputPort(ProtocolDevice, ProtocolOutputPort);
+		DMXProtocol->GetUniverseManager()->RemoveUniverse(ProtocolOutputPort, ProtocolUniverse);
+	}
+
+	const uint8 DeviceID;
+	const uint8 PortID;
+	const uint8 InterfaceID;
+	TSharedPtr<FDMXProtocolInterfaceArtNet> InterfaceArtNet;
+	TSharedPtr<FDMXProtocolDeviceArtNet> ProtocolDevice;
+	TSharedPtr<FDMXProtocolPortArtNet> ProtocolOutputPort;
+	TSharedPtr<FDMXProtocolUniverseArtNet> ProtocolUniverse;
+	FDMXProtocolArtNet* DMXProtocol;
+};
+
+IMPLEMENT_SIMPLE_AUTOMATION_TEST(FDMXPrtotocolArtNetBasicFlowTest, "VirtualProduction.DMX.ArtNet.BasicFlow", EAutomationTestFlags::ApplicationContextMask | EAutomationTestFlags::EngineFilter)
+
+static void SensDMXFragment(FDMXPrtotocolArtNetBasicFlowTest& Test, const FDMXProtocolLauncher& DMXProtocolLauncher)
+{
+	FDMXProtocolPortArtNet* ArtNetPort = static_cast<FDMXProtocolPortArtNet*>(DMXProtocolLauncher.DMXProtocol->GetPortManager()->GetPortByDeviceAndID(DMXProtocolLauncher.ProtocolDevice, DMXProtocolLauncher.PortID, EDMXPortDirection::DMX_PORT_OUTPUT));
+
+	Test.TestNotNull(TEXT("ArtNetPort is valid"), ArtNetPort);
+
+	Test.TestEqual(TEXT("Universe ID should be equal 16"), ArtNetPort->GetUniverseID(), 16);
+
+	// Send DMX fragment 1
+	const uint8 FixtureChannels[6] = { 30, 31, 32, 36, 37, 38 };
+	const uint8 FixtureValues[6] = { 255, 100, 20, 77, 88, 99 };
+	IDMXFragmentMap DMXFragmentMap;
+	DMXFragmentMap.Add(FixtureChannels[0], FixtureValues[0]);
+	DMXFragmentMap.Add(FixtureChannels[1], FixtureValues[1]);
+	DMXFragmentMap.Add(FixtureChannels[2], FixtureValues[2]);
+	DMXProtocolLauncher.DMXProtocol->SetDMXFragment(ArtNetPort->GetUniverseID(), DMXFragmentMap);
+
+	// Wait before recieve DMX
+	FPlatformProcess::Sleep(0.2f);
+	TSharedPtr<IDMXProtocolUniverse> Universe = ArtNetPort->GetCachedUniverse().Pin();
+	Test.TestTrue(TEXT("Universe should be valid"), Universe.IsValid());
+	TSharedPtr<FDMXBuffer> InputDMXBuffer = Universe->GetInputDMXBuffer();
+
+	Test.TestEqual(TEXT("Incoming buffer should be same"), FixtureValues[0], InputDMXBuffer->GetDMXData()[FixtureChannels[0]]);
+	Test.TestEqual(TEXT("Incoming buffer should be same"), FixtureValues[1], InputDMXBuffer->GetDMXData()[FixtureChannels[1]]);
+	Test.TestEqual(TEXT("Incoming buffer should be same"), FixtureValues[2], InputDMXBuffer->GetDMXData()[FixtureChannels[2]]);
+
+	// Send DMX fragment 2
+	DMXFragmentMap.Add(FixtureChannels[3], FixtureValues[3]);
+	DMXFragmentMap.Add(FixtureChannels[4], FixtureValues[4]);
+	DMXFragmentMap.Add(FixtureChannels[5], FixtureValues[5]);
+	DMXProtocolLauncher.DMXProtocol->SetDMXFragment(ArtNetPort->GetUniverseID(), DMXFragmentMap);
+
+	FPlatformProcess::Sleep(0.2f);
+
+	Test.TestEqual(TEXT("Old value DMX input should be the same"), FixtureValues[0], InputDMXBuffer->GetDMXData()[FixtureChannels[0]]);
+	Test.TestEqual(TEXT("Incoming buffer should be same"), FixtureValues[4], InputDMXBuffer->GetDMXData()[FixtureChannels[4]]);
+	Test.TestEqual(TEXT("Incoming buffer should be same"), FixtureValues[5], InputDMXBuffer->GetDMXData()[FixtureChannels[5]]);
+
+	// @TODO Implement test cases for Pool and RDM
+	// Send DMX pull
+	DMXProtocolLauncher.DMXProtocol->TransmitPoll();
+
+	// Single TOD Request
+	{
+		// Send TOD to non existent device
+		{
+			FJsonObject DeviceSettings;
+			TSharedPtr<FDMXProtocolDeviceArtNet> NonExistentDevice = MakeShared<FDMXProtocolDeviceArtNet>(DMXProtocolLauncher.DMXProtocol, DMXProtocolLauncher.InterfaceArtNet, DeviceSettings, 987);
+			DMXProtocolLauncher.DMXProtocol->TransmitTodRequest(NonExistentDevice);
+			FPlatformProcess::Sleep(0.2f);
+			const FDMXProtocolArtNetTodRequest& IncomingTodRequest = DMXProtocolLauncher.DMXProtocol->GetIncomingTodRequest();
+			Test.TestEqual(TEXT("TodRequest.AdCount should be 0"), IncomingTodRequest.AdCount, uint8(0));
+		}
+
+		// Send TOD to existing device
+		{
+			DMXProtocolLauncher.DMXProtocol->TransmitTodRequest(DMXProtocolLauncher.ProtocolDevice);
+			FPlatformProcess::Sleep(0.2f);
+			const FDMXProtocolArtNetTodRequest& IncomingTodRequest = DMXProtocolLauncher.DMXProtocol->GetIncomingTodRequest();
+			Test.TestNotEqual(TEXT("TodRequest.AdCount should not be 0"), IncomingTodRequest.AdCount, uint8(0));
+		}
+	}
+
+	// Test ArtNet RDM
+	DMXProtocolLauncher.DMXProtocol->TransmitTodRequestToAll();
+	ArtNetPort->AddTODUID(FRDMUID({ 0x00, 0x01, 0x02, 0x03, 0x04, 0x05 }));
+	ArtNetPort->AddTODUID(FRDMUID({ 0x06, 0x07, 0x08, 0x09, 0x10, 0x11 }));
+
+	FRDMUID RDMUID3;
+	uint8 RDMBuffer3[RDM_UID_WIDTH] = { 0x12, 0x13, 0x14, 0x15, 0x16, 0x17 };
+	FMemory::Memcmp(RDMUID3.Buffer, RDMBuffer3, RDM_UID_WIDTH);
+	ArtNetPort->AddTODUID(RDMUID3);
+
+	// Test wrong size of RDM Buffer C array
+	uint8 RDMBuffer4[3] = { 0x18, 0x19, 0x20 };
+	ArtNetPort->AddTODUID(FRDMUID(RDMBuffer4));
+
+	// Test wrong size of RDM Buffer C array
+	uint8 RDMBuffer5[10] = { 0x26, 0x27, 0x28, 0x29, 0x30, 0x31, 0x32, 0x32, 0x33 };
+	ArtNetPort->AddTODUID(FRDMUID(RDMBuffer5));
+
+	// Test wrong size of RDM Buffer TArray
+	ArtNetPort->AddTODUID(FRDMUID({ 0x34, 0x35 }));
+
+	DMXProtocolLauncher.DMXProtocol->TransmitTodData(DMXProtocolLauncher.ProtocolDevice, DMXProtocolLauncher.PortID);
+	DMXProtocolLauncher.DMXProtocol->TransmitTodControl(DMXProtocolLauncher.ProtocolDevice, DMXProtocolLauncher.PortID, 0x02);
+
+	TArray<uint8> RDMData;
+	RDMData.AddZeroed(ARTNET_MAX_RDM_DATA);
+	RDMData[47] = 0x00;
+	RDMData[48] = 0xcc;
+	DMXProtocolLauncher.DMXProtocol->TransmitRDM(DMXProtocolLauncher.ProtocolDevice, DMXProtocolLauncher.PortID, RDMData);
+
+}
+
+bool FDMXPrtotocolArtNetBasicFlowTest::RunTest(const FString& Parameters)
+{
+	FDMXProtocolLauncher DMXProtocolLauncher(0, 0);
+	// Test Sending DMX and RDM
+	SensDMXFragment(*this, DMXProtocolLauncher);
+
+	return true;
+}
+
+IMPLEMENT_SIMPLE_AUTOMATION_TEST(FDMXPrtotocolArtNetConsoleCommandsTest, "VirtualProduction.DMX.ArtNet.ConsoleCommands", EAutomationTestFlags::ApplicationContextMask | EAutomationTestFlags::EngineFilter)
+
+bool FDMXPrtotocolArtNetConsoleCommandsTest::RunTest(const FString& Parameters)
+{
+	FDMXProtocolLauncher DMXProtocolLauncher(0, 0);
+
+	// DMX.ArtNet.SendDMX[UniverseID] Channel:Value Channel : Value Channel : Value ...
+	uint32 Values[3][2] =
+	{
+		{25, 125},
+		{35, 25},
+		{45, 75},
+	};
+
+	FString Command(
+		FString::Printf(TEXT("DMX.ArtNet.SendDMX %d %d:%d %d:%d %d:%d"),
+			DMXProtocolLauncher.ProtocolOutputPort->GetUniverseID(),
+			Values[0][0], Values[0][1],
+			Values[1][0], Values[1][1],
+			Values[2][0], Values[2][1]
+		));
+
+	GEngine->Exec(GWorld, *Command);
+
+	FPlatformProcess::Sleep(0.2f);
+
+	TSharedPtr<IDMXProtocolUniverse> Universe = DMXProtocolLauncher.ProtocolOutputPort->GetCachedUniverse().Pin();
+	TSharedPtr<FDMXBuffer> InputDMXBuffer = Universe->GetInputDMXBuffer();
+	TestEqual(TEXT("Incoming buffer should be same"), Values[0][1], InputDMXBuffer->GetDMXData()[Values[0][0]]);
+	TestEqual(TEXT("Incoming buffer should be same"), Values[1][1], InputDMXBuffer->GetDMXData()[Values[1][0]]);
+	TestEqual(TEXT("Incoming buffer should be same"), Values[2][1], InputDMXBuffer->GetDMXData()[Values[2][0]]);
+
+	return true;
+}
\ No newline at end of file
diff --git a/Engine/Plugins/Experimental/DigitalMultiplex/DMXProtocol/Source/DMXProtocolArtNet/Public/DMXProtocolArtNet.h b/Engine/Plugins/Experimental/DigitalMultiplex/DMXProtocol/Source/DMXProtocolArtNet/Public/DMXProtocolArtNet.h
new file mode 100644
--- /dev/null
+++ b/Engine/Plugins/Experimental/DigitalMultiplex/DMXProtocol/Source/DMXProtocolArtNet/Public/DMXProtocolArtNet.h
@@ -0,0 +1,113 @@
+// Copyright 1998-2019 Epic Games, Inc. All Rights Reserved.
+
+#pragma once
+
+#include "CoreMinimal.h"
+#include "HAL/ThreadSafeCounter.h"
+#include "DMXProtocol.h"
+#include "DMXProtocolArtNetModule.h"
+#include "Packets/DMXProtocolArtNetPackets.h"
+#include "Interfaces/IPv4/IPv4Endpoint.h"
+
+class FSocket;
+
+class DMXPROTOCOLARTNET_API FDMXProtocolArtNet
+	: public FDMXProtocol
+{
+public:
+	//~ Begin IDMXProtocolBase implementation
+	virtual bool Init() override;
+	virtual bool Shutdown() override;
+	virtual bool Tick(float DeltaTime) override;
+	//~ End IDMXProtocolBase implementation
+
+	//~ Begin IDMXProtocol implementation
+	virtual TSharedPtr<IDMXProtocolSender> GetSenderInterface() const override;
+	virtual FName GetProtocolName() const override;
+	virtual void Reload() override;
+	virtual bool SendDMX(uint16 UniverseID, uint8 PortID, const TSharedPtr<FDMXBuffer>& DMXBuffer) const override;
+	virtual bool IsEnabled() const override;
+	//~ End IDMXProtocol implementation
+
+	//~ Begin IDMXProtocolRDM implementation
+	virtual void SendRDMCommand(const TSharedPtr<FJsonObject>& CMD) override;
+	virtual void RDMDiscovery(const TSharedPtr<FJsonObject>& CMD) override;
+	//~ End IDMXProtocol implementation
+
+	//~ ArtNet Transmit functions
+	bool TransmitPoll();
+	bool TransmitTodRequestToAll();
+	bool TransmitTodRequest(const TSharedPtr<IDMXProtocolDevice>& InDevice);
+
+	//~ ArtNet Getters functions
+	const FDMXProtocolArtNetTodRequest& GetIncomingTodRequest() const { return IncomingTodRequest; }
+	const FDMXProtocolArtNetTodData& GetIncomingTodData() const { return IncomingTodData; }
+	const FDMXProtocolArtNetTodControl& GetIncomingTodControl() const { return IncomingTodControl; }
+	const FDMXProtocolArtNetRDM& GetIncomingRDM() const { return IncomingRDM; }
+
+	/**
+	 * Transmit TOD data
+	 * @param InDevice Pointer for the physical device
+	 * @param InPortID The number of the port to send data for
+	 * @return true if the call was successfully placed for sending queue.
+	 */
+	bool TransmitTodData(const TSharedPtr<IDMXProtocolDevice>& InDevice, uint8 InPortID);
+
+	/**
+	 * Transmit TOD cocontrol package
+	 * @param InDevice Pointer for the physical device
+	 * @param InPortID The number of the port to send data for
+	 * @param Action tod control action
+	 * @return true if the call was successfully placed for sending queue.
+	 */
+	bool TransmitTodControl(const TSharedPtr<IDMXProtocolDevice>& InDevice, uint8 InPortID, uint8 Action);
+
+	/**
+	 * Send a RDM message
+	 * @param InDevice Pointer for the physical device
+	 * @param InPortID The number of the port to send data for
+	 * @param Data Buffer array to send
+	 * @return true if the call was successfully placed for sending queue.
+	 */
+	bool TransmitRDM(const TSharedPtr<IDMXProtocolDevice>& InDevice, uint8 InPortID, const TArray<uint8>& Data);
+
+	//~ Only the factory makes instances
+	FDMXProtocolArtNet() = delete;
+	explicit FDMXProtocolArtNet(FName InProtocolName, FJsonObject& InSettings)
+		: FDMXProtocol(FDMXProtocolArtNetModule::NAME_Artnet, InSettings)
+	{}
+
+public:
+	const TSharedPtr<FInternetAddr>& GetBroadcastAddr() { return BroadcastAddr; }
+	void OnDataReceived(const FArrayReaderPtr & Buffer);
+
+private:
+	//~ DMX Handlers
+	bool HandleReplyPacket(const FArrayReaderPtr & Buffer);
+	bool HandleDataPacket(const FArrayReaderPtr & Buffer);
+
+	//~ RDM Handlers
+	bool HandleTodRequest(const FArrayReaderPtr & Buffer);
+	bool HandleTodData(const FArrayReaderPtr & Buffer);
+	bool HandleTodControl(const FArrayReaderPtr & Buffer);
+	bool HandleRdm(const FArrayReaderPtr & Buffer);
+
+private:
+	TSharedPtr<IDMXProtocolSender> ArtNetSender;
+	TSharedPtr<IDMXProtocolReceiver> ArtNetReceiver;
+
+	/** Holds the network socket used to transport packages. */
+	FSocket* BroadcastSocket;
+
+	TSharedPtr<FInternetAddr> BroadcastAddr;
+	FIPv4Endpoint BroadcastEndpoint;
+
+	TSharedPtr<FInternetAddr> SenderAddr;
+	FIPv4Endpoint SenderEndpoint;
+
+	FDMXProtocolArtNetTodRequest IncomingTodRequest;
+	FDMXProtocolArtNetTodData IncomingTodData;
+	FDMXProtocolArtNetTodControl IncomingTodControl;
+	FDMXProtocolArtNetRDM IncomingRDM;
+};
+
diff --git a/Engine/Plugins/Experimental/DigitalMultiplex/DMXProtocol/Source/DMXProtocolArtNet/Public/DMXProtocolArtNetModule.h b/Engine/Plugins/Experimental/DigitalMultiplex/DMXProtocol/Source/DMXProtocolArtNet/Public/DMXProtocolArtNetModule.h
new file mode 100644
--- /dev/null
+++ b/Engine/Plugins/Experimental/DigitalMultiplex/DMXProtocol/Source/DMXProtocolArtNet/Public/DMXProtocolArtNetModule.h
@@ -0,0 +1,56 @@
+// Copyright 1998-2019 Epic Games, Inc. All Rights Reserved.
+
+#pragma once
+
+#include "CoreMinimal.h"
+#include "Modules/ModuleInterface.h"
+#include "IDMXProtocolFactory.h"
+
+/**
+ */
+class FDMXProtocolFactoryArtNet : public IDMXProtocolFactory
+{
+public:
+	virtual TSharedPtr<IDMXProtocol> CreateProtocol(const FName& ProtocolName) override;
+};
+
+
+/**
+ */
+class DMXPROTOCOLARTNET_API FDMXProtocolArtNetModule 
+	: public IModuleInterface
+{
+private:
+
+	/** Class responsible for creating instance(s) of the protocol */
+	TUniquePtr<FDMXProtocolFactoryArtNet> FactoryArtNet;
+
+public:
+	static FName const NAME_Artnet;
+
+public:
+	//~ Begin IModuleInterface implementation
+	virtual void StartupModule() override;
+	virtual void ShutdownModule() override;
+	//~ End IModuleInterface implementation
+
+	/** Get the instance of this module. */
+	static FDMXProtocolArtNetModule& Get();
+
+private:
+	/**
+	 * Sending DMX through console command
+	 * Command structure is DMX.SACN.SendDMX [UniverseID] Channel:Value Channel:Value Channel:Value ...
+	 * Example:
+	 * DMX.ArtNet.SendDMX 17 10:6 11:7 12:8 13:9
+	 * It will send the DMX to Universe 17. It could be any value from 0 to 32767
+	 * And it update the channels values for channel 10, 11, 12, 13 with values 6, 7, 8, 9
+	 * Channel could be any value from 0 to 511
+	 * Value could be any value from 0 to 255
+	 */
+	static void SendDMXCommandHandler(const TArray<FString>& Args);
+
+private:
+	/** Command for sending DMX through the console */
+	static FAutoConsoleCommand SendDMXCommand;
+};
diff --git a/Engine/Plugins/Experimental/DigitalMultiplex/DMXProtocol/Source/DMXProtocolSACN/DMXProtocolsACN.Build.cs b/Engine/Plugins/Experimental/DigitalMultiplex/DMXProtocol/Source/DMXProtocolSACN/DMXProtocolsACN.Build.cs
new file mode 100644
--- /dev/null
+++ b/Engine/Plugins/Experimental/DigitalMultiplex/DMXProtocol/Source/DMXProtocolSACN/DMXProtocolsACN.Build.cs
@@ -0,0 +1,32 @@
+// Copyright 1998-2019 Epic Games, Inc. All Rights Reserved.
+
+using UnrealBuildTool;
+
+public class DMXProtocolSACN : ModuleRules
+{
+	public DMXProtocolSACN(ReadOnlyTargetRules Target) : base(Target)
+	{
+		PCHUsage = ModuleRules.PCHUsageMode.UseExplicitOrSharedPCHs;
+
+        PublicDependencyModuleNames.AddRange(
+            new string[]
+            {
+                "Core",
+                "DMXProtocol"
+            }
+            );
+
+
+        PrivateDependencyModuleNames.AddRange(
+			new string[]
+			{
+				"CoreUObject",
+				"Engine",
+                "Networking",
+                "Sockets",
+                "Json"
+			}
+			);
+		
+	}
+}
diff --git a/Engine/Plugins/Experimental/DigitalMultiplex/DMXProtocol/Source/DMXProtocolSACN/Private/DMXProtocolDeviceSACN.cpp b/Engine/Plugins/Experimental/DigitalMultiplex/DMXProtocol/Source/DMXProtocolSACN/Private/DMXProtocolDeviceSACN.cpp
new file mode 100644
--- /dev/null
+++ b/Engine/Plugins/Experimental/DigitalMultiplex/DMXProtocol/Source/DMXProtocolSACN/Private/DMXProtocolDeviceSACN.cpp
@@ -0,0 +1,42 @@
+#include "DMXProtocolDeviceSACN.h"
+#include "Interfaces/IDMXProtocol.h"
+#include "Interfaces/IDMXProtocolInterface.h"
+#include "Interfaces/IDMXProtocolPort.h"
+#include "DMXProtocolTypes.h"
+#include "Json.h"
+
+FDMXProtocolDeviceSACN::FDMXProtocolDeviceSACN(IDMXProtocol* InDMXProtocol, TSharedPtr<IDMXProtocolInterface> InProtocolInterface, FJsonObject& InSettings, uint32 InDeviceID)
+	: DMXProtocol(InDMXProtocol)
+	, ProtocolInterface(InProtocolInterface)
+	, DeviceID(InDeviceID)
+{
+	Settings = MakeShared<FJsonObject>(InSettings);
+
+	checkf(DMXProtocol, TEXT("DMXProtocol pointer is nullptr"));
+	checkf(ProtocolInterface.IsValid(), TEXT("ProtocolInterface is not valid"));
+}
+
+TSharedPtr<FJsonObject> FDMXProtocolDeviceSACN::GetSettings() const
+{
+	return Settings;
+}
+
+TWeakPtr<IDMXProtocolInterface> FDMXProtocolDeviceSACN::GetCachedProtocolInterface() const
+{
+	return ProtocolInterface;
+}
+
+IDMXProtocol* FDMXProtocolDeviceSACN::GetProtocol() const
+{
+	return DMXProtocol;
+}
+
+uint32 FDMXProtocolDeviceSACN::GetDeviceID() const
+{
+	return DeviceID;
+}
+
+bool FDMXProtocolDeviceSACN::AllowLooping() const
+{
+	return true;
+}
diff --git a/Engine/Plugins/Experimental/DigitalMultiplex/DMXProtocol/Source/DMXProtocolSACN/Private/DMXProtocolDeviceSACN.h b/Engine/Plugins/Experimental/DigitalMultiplex/DMXProtocol/Source/DMXProtocolSACN/Private/DMXProtocolDeviceSACN.h
new file mode 100644
--- /dev/null
+++ b/Engine/Plugins/Experimental/DigitalMultiplex/DMXProtocol/Source/DMXProtocolSACN/Private/DMXProtocolDeviceSACN.h
@@ -0,0 +1,29 @@
+// Copyright 1998-2019 Epic Games, Inc. All Rights Reserved.
+
+#pragma once
+
+#include "CoreMinimal.h"
+#include "Interfaces/IDMXProtocolDevice.h"
+#include "Json.h"
+
+class DMXPROTOCOLSACN_API FDMXProtocolDeviceSACN
+	: public IDMXProtocolDevice
+{
+public:
+	FDMXProtocolDeviceSACN(IDMXProtocol* InDMXProtocol, TSharedPtr<IDMXProtocolInterface> InProtocolInterface, FJsonObject& InSettings, uint32 InDeviceID);
+
+	//~ Begin IDMXProtocolDevice implementation
+	virtual TSharedPtr<FJsonObject> GetSettings() const override;
+	virtual TWeakPtr<IDMXProtocolInterface> GetCachedProtocolInterface() const override;
+	virtual IDMXProtocol* GetProtocol() const override;
+	virtual uint32 GetDeviceID() const override;
+
+	virtual bool AllowLooping() const override;
+	//~ End IDMXProtocolDevice implementation
+
+private:
+	IDMXProtocol* DMXProtocol;
+	TWeakPtr<IDMXProtocolInterface> ProtocolInterface;
+	TSharedPtr<FJsonObject> Settings;
+	uint32 DeviceID;
+};
\ No newline at end of file
diff --git a/Engine/Plugins/Experimental/DigitalMultiplex/DMXProtocol/Source/DMXProtocolSACN/Private/DMXProtocolInterfaceSACN.cpp b/Engine/Plugins/Experimental/DigitalMultiplex/DMXProtocol/Source/DMXProtocolSACN/Private/DMXProtocolInterfaceSACN.cpp
new file mode 100644
--- /dev/null
+++ b/Engine/Plugins/Experimental/DigitalMultiplex/DMXProtocol/Source/DMXProtocolSACN/Private/DMXProtocolInterfaceSACN.cpp
@@ -0,0 +1,12 @@
+#include "DMXProtocolInterfaceSACN.h"
+#include "Interfaces/IDMXProtocol.h"
+
+FDMXProtocolInterfaceSACN::FDMXProtocolInterfaceSACN(IDMXProtocol* InDMXProtocol)
+	: DMXProtocol(InDMXProtocol)
+{
+}
+
+IDMXProtocol* FDMXProtocolInterfaceSACN::GetProtocol() const
+{
+	return DMXProtocol;
+}
\ No newline at end of file
diff --git a/Engine/Plugins/Experimental/DigitalMultiplex/DMXProtocol/Source/DMXProtocolSACN/Private/DMXProtocolInterfaceSACN.h b/Engine/Plugins/Experimental/DigitalMultiplex/DMXProtocol/Source/DMXProtocolSACN/Private/DMXProtocolInterfaceSACN.h
new file mode 100644
--- /dev/null
+++ b/Engine/Plugins/Experimental/DigitalMultiplex/DMXProtocol/Source/DMXProtocolSACN/Private/DMXProtocolInterfaceSACN.h
@@ -0,0 +1,20 @@
+// Copyright 1998-2019 Epic Games, Inc. All Rights Reserved.
+
+#pragma once
+
+#include "CoreMinimal.h"
+#include "Interfaces/IDMXProtocolInterface.h"
+
+class DMXPROTOCOLSACN_API FDMXProtocolInterfaceSACN
+	: public IDMXProtocolInterfaceEthernet
+{
+public:
+	FDMXProtocolInterfaceSACN(IDMXProtocol* InDMXProtocol);
+
+	//~ Begin IDMXProtocolInterface implementation
+	virtual IDMXProtocol* GetProtocol() const override;
+	//~ End IDMXProtocolInterface implementation
+
+private:
+	IDMXProtocol* DMXProtocol;
+};
\ No newline at end of file
diff --git a/Engine/Plugins/Experimental/DigitalMultiplex/DMXProtocol/Source/DMXProtocolSACN/Private/DMXProtocolPortSACN.cpp b/Engine/Plugins/Experimental/DigitalMultiplex/DMXProtocol/Source/DMXProtocolSACN/Private/DMXProtocolPortSACN.cpp
new file mode 100644
--- /dev/null
+++ b/Engine/Plugins/Experimental/DigitalMultiplex/DMXProtocol/Source/DMXProtocolSACN/Private/DMXProtocolPortSACN.cpp
@@ -0,0 +1,85 @@
+#include "DMXProtocolPortSACN.h"
+#include "Interfaces/IDMXProtocol.h"
+#include "Interfaces/IDMXProtocolUniverse.h"
+
+FDMXProtocolPortSACN::FDMXProtocolPortSACN(IDMXProtocol* InDMXProtocol, TSharedPtr<IDMXProtocolDevice> InProtocolDevice, uint8 InPortID, uint16 InUniverseID, EDMXPortDirection InPortDirection)
+	: DMXProtocol(InDMXProtocol)
+	, ProtocolDevice(InProtocolDevice)
+	, PortID(InPortID)
+	, UniverseID(InUniverseID)
+	, PortDirection(InPortDirection)
+{
+	checkf(DMXProtocol, TEXT("DMXProtocol pointer is nullptr"));
+}
+
+IDMXProtocol * FDMXProtocolPortSACN::GetProtocol() const
+{
+	return DMXProtocol;
+}
+
+uint8 FDMXProtocolPortSACN::GetPortID() const
+{
+	return PortID;
+}
+
+uint8 FDMXProtocolPortSACN::GetPriority() const
+{
+	return Priority;
+}
+
+void FDMXProtocolPortSACN::SetPriotiy(uint8 InPriority)
+{
+	Priority = InPriority;
+}
+
+bool FDMXProtocolPortSACN::IsSupportRDM() const
+{
+	return bIsRDMSupport;
+}
+
+TWeakPtr<IDMXProtocolDevice> FDMXProtocolPortSACN::GetCachedDevice() const
+{
+	return ProtocolDevice;
+}
+
+TWeakPtr<IDMXProtocolUniverse> FDMXProtocolPortSACN::GetCachedUniverse() const
+{
+	return ProtocolUniverse;
+}
+
+void FDMXProtocolPortSACN::SetUniverse(const TSharedPtr<IDMXProtocolUniverse>& InUniverse)
+{
+	ProtocolUniverse = InUniverse;
+}
+
+EDMXPortCapability FDMXProtocolPortSACN::GetPortCapability() const
+{
+	return DMX_PORT_CAPABILITY_FULL;
+}
+
+EDMXPortDirection FDMXProtocolPortSACN::GetPortDirection() const
+{
+	return PortDirection;
+}
+
+bool FDMXProtocolPortSACN::WriteDMX(const TSharedPtr<FDMXBuffer>& DMXBuffer)
+{
+	return DMXProtocol->SendDMX(UniverseID, PortID, DMXBuffer);
+}
+
+bool FDMXProtocolPortSACN::ReadDMX()
+{
+	// Nothing for now, but there is a reader in Universe
+	return false;
+}
+
+TSharedPtr<FJsonObject> FDMXProtocolPortSACN::GetSettings() const
+{
+	return Settings;
+}
+
+uint16 FDMXProtocolPortSACN::GetUniverseID() const
+{
+	return UniverseID;
+}
+
diff --git a/Engine/Plugins/Experimental/DigitalMultiplex/DMXProtocol/Source/DMXProtocolSACN/Private/DMXProtocolPortSACN.h b/Engine/Plugins/Experimental/DigitalMultiplex/DMXProtocol/Source/DMXProtocolSACN/Private/DMXProtocolPortSACN.h
new file mode 100644
--- /dev/null
+++ b/Engine/Plugins/Experimental/DigitalMultiplex/DMXProtocol/Source/DMXProtocolSACN/Private/DMXProtocolPortSACN.h
@@ -0,0 +1,44 @@
+// Copyright 1998-2019 Epic Games, Inc. All Rights Reserved.
+
+#pragma once
+
+#include "CoreMinimal.h"
+#include "Interfaces/IDMXProtocolPort.h"
+
+class IDMXProtocol;
+
+class DMXPROTOCOLSACN_API FDMXProtocolPortSACN
+	: public IDMXProtocolPort
+{
+public:
+	FDMXProtocolPortSACN(IDMXProtocol* InDMXProtocol, TSharedPtr<IDMXProtocolDevice> InProtocolDevice, uint8 InPortID, uint16 InUniverseID, EDMXPortDirection InPortDirection);
+
+	//~ Begin IDMXProtocolPort implementation
+	virtual IDMXProtocol* GetProtocol() const override;
+
+	virtual uint8 GetPortID() const override;
+	virtual uint8 GetPriority() const override;
+	virtual void SetPriotiy(uint8 InPriority) override;
+	virtual bool IsSupportRDM() const override;
+	virtual TWeakPtr<IDMXProtocolDevice> GetCachedDevice() const override;
+	virtual TWeakPtr<IDMXProtocolUniverse> GetCachedUniverse() const override;
+	virtual void SetUniverse(const TSharedPtr<IDMXProtocolUniverse>& InUniverse) override;
+	virtual EDMXPortCapability GetPortCapability() const override;
+	virtual EDMXPortDirection GetPortDirection() const override;
+	virtual bool WriteDMX(const TSharedPtr<FDMXBuffer>& DMXBuffer) override;
+	virtual bool ReadDMX() override;
+	virtual TSharedPtr<FJsonObject> GetSettings() const override;
+	virtual uint16 GetUniverseID() const override;
+	//~ End IDMXProtocolPort implementation
+
+private:
+	IDMXProtocol* DMXProtocol;
+	TWeakPtr<IDMXProtocolDevice> ProtocolDevice;
+	TWeakPtr<IDMXProtocolUniverse> ProtocolUniverse;
+	uint8 PortID;
+	uint16 UniverseID;
+	uint8 Priority;
+	bool bIsRDMSupport;
+	EDMXPortDirection PortDirection;
+	TSharedPtr<FJsonObject> Settings;
+};
diff --git a/Engine/Plugins/Experimental/DigitalMultiplex/DMXProtocol/Source/DMXProtocolSACN/Private/DMXProtocolSACN.cpp b/Engine/Plugins/Experimental/DigitalMultiplex/DMXProtocol/Source/DMXProtocolSACN/Private/DMXProtocolSACN.cpp
new file mode 100644
--- /dev/null
+++ b/Engine/Plugins/Experimental/DigitalMultiplex/DMXProtocol/Source/DMXProtocolSACN/Private/DMXProtocolSACN.cpp
@@ -0,0 +1,193 @@
+#include "DMXProtocolSACN.h"
+#include "DMXProtocolTransportSACN.h"
+#include "Common/UdpSocketBuilder.h"
+
+#include "Packets/DMXProtocolE131PDUPacket.h"
+#include "DMXProtocolPackager.h"
+#include "DMXProtocolSettings.h"
+
+#include "Serialization/BufferArchive.h"
+
+TSharedPtr<IDMXProtocolSender> FDMXProtocolSACN::GetSenderInterface() const
+{
+	return SACNSender;
+}
+
+bool FDMXProtocolSACN::Init()
+{	
+	//Socket;
+	ISocketSubsystem* SocketSubsystem = ISocketSubsystem::Get(PLATFORM_SOCKETSUBSYSTEM);
+	bool bIsValid = false;
+
+	// Create broadcast address
+	FIPv4Endpoint UnicastEndpoint = FIPv4Endpoint::Any;
+
+	// Create sender socket
+	SenderSocket = FUdpSocketBuilder(TEXT("UDPSACNSenderSocket"))
+		.AsNonBlocking()
+#if PLATFORM_WINDOWS
+		.BoundToAddress(UnicastEndpoint.Address)
+#endif
+		.AsReusable();
+
+	if (SenderSocket == nullptr)
+	{
+		UE_LOG_DMXPROTOCOL(Error, TEXT("ERROR create SenderSocket"));
+	}
+
+	SACNSender = MakeShared<FDMXProtocolSenderSACN>(*SenderSocket, this);
+
+	return true;
+}
+
+bool FDMXProtocolSACN::Shutdown()
+{
+	SACNSender.Reset();
+
+	//Clear all sockets!
+	if (SenderSocket)
+	{
+		SenderSocket->Close();
+		ISocketSubsystem::Get(PLATFORM_SOCKETSUBSYSTEM)->DestroySocket(SenderSocket);
+	}
+
+	return false;
+}
+
+bool FDMXProtocolSACN::IsEnabled() const
+{
+	return true;
+}
+
+FName FDMXProtocolSACN::GetProtocolName() const
+{
+	return FDMXProtocolSACNModule::NAME_SACN;
+}
+
+void FDMXProtocolSACN::Reload()
+{
+}
+
+bool FDMXProtocolSACN::Tick(float DeltaTime)
+{
+	return true;
+}
+
+bool FDMXProtocolSACN::SendDiscovery(const TArray<uint16>& Universes)
+{
+	// Init Packager
+	FDMXProtocolPackager Packager;
+
+	FDMXProtocolE131RootLayerPacket RootLayer;
+	RootLayer.Flags = 0x72; // Const for now
+	RootLayer.Length = 0x6e; // Const for now
+	FGuid Guid = FGuid::NewGuid();
+	FMemory::Memcpy(RootLayer.CID, &Guid, ACN_CIDBYTES);
+
+	FDMXProtocolUDPE131FramingLayerPacket FramingLayer;
+	FDMXProtocolUDPE131DiscoveryLayerPacket DiscoveryLayer;
+	if (Universes.Num() == ACN_DMX_SIZE)
+	{
+		FMemory::Memcpy(DiscoveryLayer.Universes, Universes.GetData(), ACN_DMX_SIZE);
+	}
+	else
+	{
+		UE_LOG_DMXPROTOCOL(Warning, TEXT("Size of outcoming sACN DMX discovery buffer is wrong"));
+	}
+
+	Packager.AddToPackage(&RootLayer);
+	Packager.AddToPackage(&FramingLayer);
+	Packager.AddToPackage(&DiscoveryLayer);
+
+	// Sending
+	FJsonObject PacketSettings;
+	PacketSettings.SetNumberField("UniverseID", ACN_MAX_UNIVERSES);
+	FDMXPacketPtr Packet = MakeShared<FDMXPacket, ESPMode::ThreadSafe>(PacketSettings, Packager.GetBuffer());
+	GetSenderInterface()->EnqueueOutboundPackage(Packet);
+
+	return true;
+}
+
+bool FDMXProtocolSACN::SendDMX(uint16 UniverseID, uint8 PortID, const TSharedPtr<FDMXBuffer>& DMXBuffer) const
+{
+	// Init Packager
+	FDMXProtocolPackager Packager;
+
+	// SACN PDU packets
+	FDMXProtocolE131RootLayerPacket RootLayer;;
+	RootLayer.PreambleSize = ACN_RLP_PREAMBLE_SIZE;
+	RootLayer.PostambleSize = 0x0000;
+	RootLayer.Flags = 0x72; // Const for now
+	RootLayer.Length = 0x6e; // Const for now
+	RootLayer.Vector = 4; // Const for now
+	FGuid Guid = FGuid::NewGuid();
+	FMemory::Memcpy(RootLayer.CID, &Guid, ACN_CIDBYTES);
+
+	FDMXProtocolE131FramingLayerPacket FramingLayer;
+	FramingLayer.Flags = 0x72; // Const for now
+	FramingLayer.Length = 0x58; // Const for now
+	FramingLayer.Priority = 100; // Const for now
+	FramingLayer.SynchronizationAddress = 0; // Const for now
+	FramingLayer.SequenceNumber = 0; // Const for now
+	FramingLayer.Options = 0; // Const for now
+	FramingLayer.Universe = UniverseID;
+
+	FDMXProtocolE131DMPLayerPacket DMPLayer;
+	if (DMXBuffer->GetDMXData().Num() == ACN_DMX_SIZE)
+	{
+		FMemory::Memcpy(DMPLayer.DMX, DMXBuffer->GetDMXData().GetData(), ACN_DMX_SIZE);
+	}
+	else
+	{
+		UE_LOG_DMXPROTOCOL(Warning, TEXT("Size of outcoming sACN DMX buffer is wrong"));
+	}
+	DMPLayer.Flags = 0x72; // Const for now
+	DMPLayer.Length = 0x0b; // Const for now
+	DMPLayer.AddressTypeAndDataType = 0xa1; // Const for now
+	DMPLayer.FirstPropertyAddress = 0; // Const for now
+	DMPLayer.AddressIncrement = ACN_ADDRESS_INC;
+	DMPLayer.PropertyValueCount = ACN_DMX_SIZE + 1;
+	DMPLayer.STARTCode = 0; // Const for now
+
+	Packager.AddToPackage(&RootLayer);
+	Packager.AddToPackage(&FramingLayer);
+	Packager.AddToPackage(&DMPLayer);
+
+	// Sending
+	FJsonObject PacketSettings;
+	PacketSettings.SetNumberField("UniverseID", UniverseID);
+	FDMXPacketPtr Packet = MakeShared<FDMXPacket, ESPMode::ThreadSafe>(PacketSettings, Packager.GetBuffer());
+	GetSenderInterface()->EnqueueOutboundPackage(Packet);
+
+	return false;
+}
+
+void FDMXProtocolSACN::SendRDMCommand(const TSharedPtr<FJsonObject>& CMD)
+{
+	// Cross protocol RDM send command implementation will be here
+}
+
+void FDMXProtocolSACN::RDMDiscovery(const TSharedPtr<FJsonObject>& CMD)
+{
+	// Cross protocol RDM discovery implementation will be here
+}
+
+TSharedPtr<FInternetAddr> FDMXProtocolSACN::GetUniverseAddr(uint16 InUniverseID)
+{
+	ISocketSubsystem* SocketSubsystem = ISocketSubsystem::Get(PLATFORM_SOCKETSUBSYSTEM);
+
+	TSharedPtr<FInternetAddr> InternetAddr = SocketSubsystem->CreateInternetAddr();
+	InternetAddr->SetPort(ACN_PORT);
+
+	FBufferArchive IP;
+	uint8 IP_0 = ACN_UNIVERSE_IP_0;
+	uint8 IP_1 = ACN_UNIVERSE_IP_1;
+	IP << IP_0; // [x.?.?.?]
+	IP << IP_1; // [x.x.?.?]
+	IP.ByteSwap(&InUniverseID, 2);
+	IP << InUniverseID;	// [x.x.x.x]
+
+	InternetAddr->SetRawIp(IP);
+
+	return InternetAddr;
+}
\ No newline at end of file
diff --git a/Engine/Plugins/Experimental/DigitalMultiplex/DMXProtocol/Source/DMXProtocolSACN/Private/DMXProtocolSACNConstants.h b/Engine/Plugins/Experimental/DigitalMultiplex/DMXProtocol/Source/DMXProtocolSACN/Private/DMXProtocolSACNConstants.h
new file mode 100644
--- /dev/null
+++ b/Engine/Plugins/Experimental/DigitalMultiplex/DMXProtocol/Source/DMXProtocolSACN/Private/DMXProtocolSACNConstants.h
@@ -0,0 +1,64 @@
+// Copyright 1998-2019 Epic Games, Inc. All Rights Reserved.
+
+#pragma once
+
+#include "CoreMinimal.h"
+#include "DMXProtocolTypes.h"
+
+enum 
+{
+	VECTOR_ROOT_E131_DATA = 0x00000004,
+	VECTOR_ROOT_E131_EXTENDED = 0x00000008,
+	VECTOR_E131_EXTENDED_SYNCHRONIZATION = 0x00000001,
+	VECTOR_E131_EXTENDED_DISCOVERY = 0x00000002,
+	VECTOR_E131_DATA_PACKET = 0x00000002,
+	VECTOR_DMP_SET_PROPERTY = 0x02
+};
+
+enum
+{
+	ACN_PORT = 5568,
+	ACN_UNIVERSE_IP_0 = 239,
+	ACN_UNIVERSE_IP_1 = 255
+};
+
+enum
+{
+	ACN_MAX_UNIVERSES = 64214
+};
+
+enum
+{
+	ACN_HEADER_SIZE = 126,
+	ACN_DMX_SIZE = 512,
+	ACN_IDENTIFIER_SIZE = 12,
+	ACN_RLP_PREAMBLE_SIZE = 16,
+	ACN_CIDBYTES = 16,
+	ACN_SOURCE_NAME_SIZE = 64,
+	ACN_ADDRESS_INC = 1,
+	ACN_ADDRESS_ROOT_VECTOR = 18
+};
+
+enum
+{
+	ACN_DMX_PACKAGE_SIZE = 638,
+	ACN_DMX_ROOT_PACKAGE_SIZE = 38,
+	ACN_DMX_PDU_FRAMING_PACKAGE_SIZE = 77,
+	ACN_DMX_PDU_DMP_PACKAGE_SIZE = 523,
+	ACN_DMX_UDP_FRAMING_PACKAGE_SIZE = 74,
+	ACN_DMX_UDP_DISCOVERY_PACKAGE_SIZE = 1032
+};
+
+enum
+{
+	VHD_L_FLAG = 0x80,
+	VHD_V_FLAG = 0x40,
+	VHD_H_FLAG = 0x20,
+	VHD_D_FLAG = 0x10
+};
+
+struct FSACNConstants
+{
+	static uint8 constexpr ACN_IDENTIFIER[ACN_IDENTIFIER_SIZE] = { 'A', 'S', 'C', '-', 'E', '1', '.', '1', '7', '\0', '\0', '\0' };
+};
+
diff --git a/Engine/Plugins/Experimental/DigitalMultiplex/DMXProtocol/Source/DMXProtocolSACN/Private/DMXProtocolSACNModule.cpp b/Engine/Plugins/Experimental/DigitalMultiplex/DMXProtocol/Source/DMXProtocolSACN/Private/DMXProtocolSACNModule.cpp
new file mode 100644
--- /dev/null
+++ b/Engine/Plugins/Experimental/DigitalMultiplex/DMXProtocol/Source/DMXProtocolSACN/Private/DMXProtocolSACNModule.cpp
@@ -0,0 +1,134 @@
+// Copyright 1998-2019 Epic Games, Inc. All Rights Reserved.
+
+#include "DMXProtocolSACNModule.h"
+#include "CoreMinimal.h"
+#include "Modules/ModuleManager.h"
+
+#include "Interfaces/IDMXProtocol.h"
+#include "DMXProtocolSACN.h"
+#include "DMXProtocolTypes.h"
+#include "Json.h"
+
+#include "Managers/DMXProtocolDeviceManager.h"
+#include "Managers/DMXProtocolInterfaceManager.h"
+#include "Managers/DMXProtocolPortManager.h"
+#include "Managers/DMXProtocolUniverseManager.h"
+
+#include "DMXProtocolInterfaceSACN.h"
+#include "DMXProtocolDeviceSACN.h"
+#include "DMXProtocolUniverseSACN.h"
+#include "DMXProtocolPortSACN.h"
+#include "DMXProtocolSACNConstants.h"
+
+IMPLEMENT_MODULE(FDMXProtocolSACNModule, DMXProtocolSACN);
+
+const FName FDMXProtocolSACNModule::NAME_SACN = FName(TEXT("SACN"));
+
+FAutoConsoleCommand FDMXProtocolSACNModule::SendDMXCommand(
+	TEXT("DMX.SACN.SendDMX"),
+	TEXT("Command for sending DMX through SACN Protocol. DMX.SACN.SendDMX [UniverseID] Channel:Value Channel:Value Channel:Value n\t DMX.SACN.SendDMX 17 10:6 11:7 12:8 13:9 n\t It will send channels values to the DMX to Universe 17"),
+	FConsoleCommandWithArgsDelegate::CreateStatic(&FDMXProtocolSACNModule::SendDMXCommandHandler)
+);
+
+TSharedPtr<IDMXProtocol> FDMXProtocolFactorySACN::CreateProtocol(const FName & ProtocolName)
+{
+	FJsonObject ProtocolSettings;
+	TSharedPtr<IDMXProtocol> ProtocolSACNPtr = MakeShared<FDMXProtocolSACN>(ProtocolName, ProtocolSettings);
+	if (ProtocolSACNPtr->IsEnabled())
+	{
+		if (!ProtocolSACNPtr->Init())
+		{
+			UE_LOG_DMXPROTOCOL(Warning, TEXT("SACN failed to initialize!"));
+			ProtocolSACNPtr->Shutdown();
+			ProtocolSACNPtr = nullptr;
+		}
+	}
+	else
+	{
+		UE_LOG_DMXPROTOCOL(Warning, TEXT("SACN disabled!"));
+		ProtocolSACNPtr->Shutdown();
+		ProtocolSACNPtr = nullptr;
+	}
+
+	return ProtocolSACNPtr;
+}
+
+void FDMXProtocolSACNModule::StartupModule()
+{
+	FactorySACN = MakeUnique<FDMXProtocolFactorySACN>();
+
+	// Create and register our singleton factory with the main online subsystem for easy access
+	FDMXProtocolModule& DMXProtocolModule = FModuleManager::GetModuleChecked<FDMXProtocolModule>("DMXProtocol");
+	DMXProtocolModule.RegisterProtocol(NAME_SACN, FactorySACN.Get());
+}
+
+void FDMXProtocolSACNModule::ShutdownModule()
+{
+	// Unregister and destroy protocol
+	FDMXProtocolModule* DMXProtocolModule = FModuleManager::GetModulePtr<FDMXProtocolModule>("DMXProtocol");
+	if (DMXProtocolModule != nullptr)
+	{
+		DMXProtocolModule->UnregisterProtocol(NAME_SACN);
+	}
+
+	FactorySACN.Release();
+}
+
+FDMXProtocolSACNModule& FDMXProtocolSACNModule::Get()
+{
+	return FModuleManager::GetModuleChecked<FDMXProtocolSACNModule>("DMXProtocolSACN");
+}
+
+void FDMXProtocolSACNModule::SendDMXCommandHandler(const TArray<FString>& Args)
+{
+	if (Args.Num() < 2)
+	{
+		UE_LOG_DMXPROTOCOL(Warning, TEXT("Not enough arguments. It won't be sent\n"
+			"Command structure is DMX.SACN.SendDMX [UniverseID] Channel:Value Channel:Value Channel:Value\n"
+			"For example: DMX.SACN.SendDMX 17 10:6 11:7 12:8 13:9"));
+		return;
+	}
+
+	uint32 UniverseID = 0;
+	LexTryParseString<uint32>(UniverseID, *Args[0]);
+	if (UniverseID > ACN_MAX_UNIVERSES)
+	{
+		UE_LOG_DMXPROTOCOL(Warning, TEXT("The UniverseID is bigger then the max universe value. It won't be sent. It won't be sent\n"
+			"For example: DMX.SACN.SendDMX 17 10:6 11:7 12:8 13:9\n"
+			"Where Universe %d should be less then %d"), UniverseID, ACN_MAX_UNIVERSES);
+		return;
+	}
+
+	IDMXFragmentMap DMXFragment;
+	for (int32 i = 1; i < Args.Num(); i++)
+	{
+		FString ChannelAndValue = Args[i];
+
+		FString KeyStr;
+		FString ValueStr;
+		ChannelAndValue.Split(TEXT(":"), &KeyStr, &ValueStr);
+
+		uint32 Key = 0;
+		uint32 Value = 0;
+		LexTryParseString<uint32>(Key, *KeyStr);
+		if (Key > DMX_UNIVERSE_SIZE)
+		{
+			UE_LOG_DMXPROTOCOL(Warning, TEXT("The input channel is bigger then the universe size. It won't be sent\n"
+				"For example: DMX.SACN.SendDMX 17 10:6 11:7 12:8 13:9\n"
+				"Where channel %d should be less then %d"), Key, DMX_UNIVERSE_SIZE);
+			return;
+		}
+		LexTryParseString<uint32>(Value, *ValueStr);
+		if (Value > DMX_MAX_CHANNEL_VALUE)
+		{
+			UE_LOG_DMXPROTOCOL(Warning, TEXT("The input value is bigger then the universe size. It won't be sent\n"
+				"For example: DMX.SACN.SendDMX 17 10:6 11:7 12:8 13:9\n"
+				"Where value %d should be less then %d"), Value, DMX_MAX_CHANNEL_VALUE);
+			return;
+		}
+		DMXFragment.Add(Key, Value);
+	}
+
+	FDMXProtocolSACN* DMXProtocol = IDMXProtocol::Get<FDMXProtocolSACN>(FDMXProtocolSACNModule::NAME_SACN);
+	DMXProtocol->SetDMXFragment(UniverseID, DMXFragment, true);
+}
diff --git a/Engine/Plugins/Experimental/DigitalMultiplex/DMXProtocol/Source/DMXProtocolSACN/Private/DMXProtocolSACNUtils.h b/Engine/Plugins/Experimental/DigitalMultiplex/DMXProtocol/Source/DMXProtocolSACN/Private/DMXProtocolSACNUtils.h
new file mode 100644
--- /dev/null
+++ b/Engine/Plugins/Experimental/DigitalMultiplex/DMXProtocol/Source/DMXProtocolSACN/Private/DMXProtocolSACNUtils.h
@@ -0,0 +1,29 @@
+// Copyright 1998-2019 Epic Games, Inc. All Rights Reserved.
+
+#pragma once
+
+#include "CoreMinimal.h"
+#include "DMXProtocolTypes.h"
+#include "DMXProtocolSACNConstants.h"
+
+namespace SACN
+{
+	static uint32 GetRootPacketType(const FArrayReaderPtr& Buffer)
+	{
+
+		uint32 Vector = 0x00000000;
+		const uint32 MinCheck = ACN_ADDRESS_ROOT_VECTOR + 4;
+		if (Buffer->Num() > MinCheck)
+		{
+			// Get OpCode
+			Buffer->Seek(ACN_ADDRESS_ROOT_VECTOR);
+			*Buffer << Vector;
+			Buffer->ByteSwap(&Vector, 4);
+
+			// Reset Position
+			Buffer->Seek(0);
+		}
+
+		return Vector;
+	}
+}
\ No newline at end of file
diff --git a/Engine/Plugins/Experimental/DigitalMultiplex/DMXProtocol/Source/DMXProtocolSACN/Private/DMXProtocolTransportSACN.cpp b/Engine/Plugins/Experimental/DigitalMultiplex/DMXProtocol/Source/DMXProtocolSACN/Private/DMXProtocolTransportSACN.cpp
new file mode 100644
--- /dev/null
+++ b/Engine/Plugins/Experimental/DigitalMultiplex/DMXProtocol/Source/DMXProtocolSACN/Private/DMXProtocolTransportSACN.cpp
@@ -0,0 +1,212 @@
+// Copyright 1998-2019 Epic Games, Inc. All Rights Reserved.
+
+#include "DMXProtocolTransportSACN.h"
+#include "DMXProtocolSACN.h"
+
+#include "Common/UdpSocketSender.h"
+#include "HAL/Event.h"
+#include "HAL/Runnable.h"
+#include "HAL/RunnableThread.h"
+#include "IMessageAttachment.h"
+#include "Serialization/ArrayReader.h"
+#include "Serialization/ArrayWriter.h"
+#include "Sockets.h"
+
+static FTimespan CalculateWaitTime()
+{
+	return FTimespan::FromMilliseconds(10);
+}
+
+FDMXProtocolSenderSACN::FDMXProtocolSenderSACN(FSocket& InSocket, FDMXProtocolSACN* InProtocol)
+	: LastSentPackage(-1)
+	, bRequestingExit(false)
+	, BroadcastSocket(&InSocket)
+	, Protocol(InProtocol)
+{
+	WorkEvent = MakeShareable(FPlatformProcess::GetSynchEventFromPool(), [](FEvent* EventToDelete)
+	{
+		FPlatformProcess::ReturnSynchEventToPool(EventToDelete);
+	});
+
+	Thread = FRunnableThread::Create(this, TEXT("FDMXProtocolSenderSACN"), 128 * 1024, TPri_BelowNormal, FPlatformAffinity::GetPoolThreadMask());
+}
+
+FDMXProtocolSenderSACN::~FDMXProtocolSenderSACN()
+{
+	// Shut Down worker thread
+	if (Thread != nullptr)
+	{
+		Thread->Kill(true);
+		delete Thread;
+	}
+}
+
+bool FDMXProtocolSenderSACN::Init()
+{
+	return true;
+}
+
+uint32 FDMXProtocolSenderSACN::Run()
+{
+	while (StopTaskCounter.GetValue() == 0)
+	{
+		if (WorkEvent->Wait(CalculateWaitTime()))
+		{
+			ConsumeOutboundPackages();
+		}
+	}
+
+	return 0;
+}
+
+void FDMXProtocolSenderSACN::Stop()
+{
+	StopTaskCounter.Increment();
+
+	if (WorkEvent)
+	{
+		WorkEvent->Trigger();
+		Thread->WaitForCompletion();
+	}
+}
+
+void FDMXProtocolSenderSACN::Exit()
+{
+}
+
+void FDMXProtocolSenderSACN::Tick()
+{
+	ConsumeOutboundPackages();
+}
+
+FSingleThreadRunnable* FDMXProtocolSenderSACN::GetSingleThreadInterface()
+{
+	return this;
+}
+
+bool FDMXProtocolSenderSACN::EnqueueOutboundPackage(FDMXPacketPtr Packet)
+{
+	if (StopTaskCounter.GetValue())
+	{
+		return false;
+	}
+
+	if (!OutboundPackages.Enqueue(Packet))
+	{
+		return false;
+	}
+
+	WorkEvent->Trigger();
+
+
+	return true;
+}
+
+void FDMXProtocolSenderSACN::ConsumeOutboundPackages()
+{
+	FScopeLock ScopeLock(&CriticalSection);
+
+	FDMXPacketPtr Packet;
+	while (OutboundPackages.Dequeue(Packet))
+	{
+		++LastSentPackage;
+		int32 BytesSent = 0;
+		bool bIsSent = BroadcastSocket->SendTo(Packet->Data.GetData(), Packet->Data.Num(), BytesSent, *FDMXProtocolSACN::GetUniverseAddr(Packet->Settings.GetNumberField(TEXT("UniverseID"))));
+		if (!bIsSent)
+		{
+			UE_LOG_DMXPROTOCOL(Error, TEXT("Error sending %d"), bIsSent);
+		}
+	}
+}
+
+FDMXProtocolReceiverSACN::FDMXProtocolReceiverSACN(FSocket& InSocket, const FTimespan& InWaitTime)
+	: Socket(&InSocket)
+	, Stopping(false)
+	, Thread(nullptr)
+	, WaitTime(InWaitTime)
+{
+	check(Socket != nullptr);
+	check(Socket->GetSocketType() == SOCKTYPE_Datagram);
+
+	SocketSubsystem = ISocketSubsystem::Get(PLATFORM_SOCKETSUBSYSTEM);
+
+	Thread = FRunnableThread::Create(this, *ThreadName, 128 * 1024, TPri_Normal, FPlatformAffinity::GetPoolThreadMask());
+}
+
+FDMXProtocolReceiverSACN::~FDMXProtocolReceiverSACN()
+{
+	if (Thread != nullptr)
+	{
+		Thread->Kill(true);
+		delete Thread;
+	}
+}
+
+FOnDMXDataReceived& FDMXProtocolReceiverSACN::OnDataReceived()
+{
+	return DMXDataReceiveDelegate;
+}
+
+FRunnableThread* FDMXProtocolReceiverSACN::GetThread() const
+{
+	return Thread;
+}
+
+bool FDMXProtocolReceiverSACN::Init()
+{
+	return true;
+}
+
+uint32 FDMXProtocolReceiverSACN::Run()
+{
+	while (!Stopping)
+	{
+		Update(WaitTime);
+	}
+
+	return 0;
+}
+
+void FDMXProtocolReceiverSACN::Stop()
+{
+	Stopping = true;
+}
+
+void FDMXProtocolReceiverSACN::Exit()
+{
+}
+
+void FDMXProtocolReceiverSACN::Tick()
+{
+	Update(FTimespan::Zero());
+}
+
+FSingleThreadRunnable* FDMXProtocolReceiverSACN::GetSingleThreadInterface()
+{
+	return this;
+}
+
+void FDMXProtocolReceiverSACN::Update(const FTimespan& SocketWaitTime)
+{
+	if (!Socket->Wait(ESocketWaitConditions::WaitForRead, SocketWaitTime))
+	{
+		return;
+	}
+
+	TSharedRef<FInternetAddr> Sender = SocketSubsystem->CreateInternetAddr();
+	uint32 Size;
+
+	while (Socket->HasPendingData(Size))
+	{
+		FArrayReaderPtr Reader = MakeShareable(new FArrayReader(true));
+		Reader->SetNumUninitialized(FMath::Min(Size, 65507u));
+
+		int32 Read = 0;
+
+		if (Socket->RecvFrom(Reader->GetData(), Reader->Num(), Read, *Sender))
+		{
+			Reader->RemoveAt(Read, Reader->Num() - Read, false);
+			DMXDataReceiveDelegate.ExecuteIfBound(Reader);
+		}
+	}
+}
diff --git a/Engine/Plugins/Experimental/DigitalMultiplex/DMXProtocol/Source/DMXProtocolSACN/Private/DMXProtocolTransportSACN.h b/Engine/Plugins/Experimental/DigitalMultiplex/DMXProtocol/Source/DMXProtocolSACN/Private/DMXProtocolTransportSACN.h
new file mode 100644
--- /dev/null
+++ b/Engine/Plugins/Experimental/DigitalMultiplex/DMXProtocol/Source/DMXProtocolSACN/Private/DMXProtocolTransportSACN.h
@@ -0,0 +1,123 @@
+// Copyright 1998-2019 Epic Games, Inc. All Rights Reserved.
+
+#pragma once
+
+#include "CoreMinimal.h"
+#include "Containers/Queue.h"
+#include "Interfaces/IDMXProtocolTransport.h"
+#include "DMXProtocolTypes.h"
+
+
+class FSocket;
+class FDMXProtocolSACN;
+
+
+class DMXPROTOCOLSACN_API FDMXProtocolSenderSACN
+	: public IDMXProtocolSender
+{
+
+public:
+	FDMXProtocolSenderSACN(FSocket& InSocket, FDMXProtocolSACN* InProtocol);
+	virtual ~FDMXProtocolSenderSACN();
+
+	//~ Begin FRunnable implementation
+	virtual bool Init() override;
+	virtual uint32 Run() override;
+	virtual void Stop() override;
+	virtual void Exit() override;
+	//~ End FRunnable implementation
+
+	//~ Begin FSingleThreadRunnable implementation
+	virtual void Tick() override;
+	virtual class FSingleThreadRunnable* GetSingleThreadInterface() override;
+	//~ End FSingleThreadRunnable implementation
+
+	//~ Begin IDMXProtocolSender implementation
+	virtual bool EnqueueOutboundPackage(FDMXPacketPtr Packet) override;
+	//~ End IDMXProtocolSender implementation
+
+public:
+	/** Consumes all outbound packages. */
+	void ConsumeOutboundPackages();
+
+
+private:
+	/** Holds the queue of outbound packages. */
+	TQueue<FDMXPacketPtr, EQueueMode::Mpsc> OutboundPackages;
+
+	/** Holds the last sent message number. */
+	int32 LastSentPackage;
+
+	/** */
+	mutable FCriticalSection CriticalSection;
+
+	FThreadSafeCounter StopTaskCounter;
+
+	/** Holds the thread object. */
+	FRunnableThread* Thread;
+
+	/** Holds an event signaling that inbound messages need to be processed. */
+	TSharedPtr<FEvent, ESPMode::ThreadSafe> WorkEvent;
+
+	FThreadSafeBool bRequestingExit;
+
+	/** Holds the network socket used to sender packages. */
+	FSocket* BroadcastSocket;
+
+	FDMXProtocolSACN* Protocol;
+};
+
+
+class FDMXProtocolReceiverSACN
+	: public IDMXProtocolReceiver
+{
+public:
+	FDMXProtocolReceiverSACN(FSocket& InSocket, const FTimespan& InWaitTime);
+	virtual ~FDMXProtocolReceiverSACN();
+
+public:
+	//~ Begin FRunnable implementation
+	virtual bool Init() override;
+	virtual uint32 Run() override;
+	virtual void Stop() override;
+	virtual void Exit() override;
+	//~ End FRunnable implementation
+
+	//~ Begin FSingleThreadRunnable implementation
+	virtual void Tick() override;
+	virtual class FSingleThreadRunnable* GetSingleThreadInterface() override;
+	//~ End FSingleThreadRunnable implementation
+
+	//~ Begin IDMXProtocolReceiver implementation
+	virtual FOnDMXDataReceived& OnDataReceived() override;
+	virtual FRunnableThread* GetThread() const override;
+	//~ End IDMXProtocolReceiver implementation
+
+protected:
+	void Update(const FTimespan& SocketWaitTime);
+
+private:
+
+	/** The network socket. */
+	FSocket* Socket;
+
+	/** Pointer to the socket sub-system. */
+	ISocketSubsystem* SocketSubsystem;
+
+	/** Flag indicating that the thread is stopping. */
+	bool Stopping;
+
+	/** The thread object. */
+	FRunnableThread* Thread;
+
+	/** The receiver thread's name. */
+	FString ThreadName;
+
+	/** The amount of time to wait for inbound packets. */
+	FTimespan WaitTime;
+
+private:
+
+	/** Holds the data received delegate. */
+	FOnDMXDataReceived DMXDataReceiveDelegate;
+};
\ No newline at end of file
diff --git a/Engine/Plugins/Experimental/DigitalMultiplex/DMXProtocol/Source/DMXProtocolSACN/Private/DMXProtocolUniverseSACN.cpp b/Engine/Plugins/Experimental/DigitalMultiplex/DMXProtocol/Source/DMXProtocolSACN/Private/DMXProtocolUniverseSACN.cpp
new file mode 100644
--- /dev/null
+++ b/Engine/Plugins/Experimental/DigitalMultiplex/DMXProtocol/Source/DMXProtocolSACN/Private/DMXProtocolUniverseSACN.cpp
@@ -0,0 +1,133 @@
+#include "DMXProtocolUniverseSACN.h"
+
+#include "DMXProtocolSACN.h"
+#include "Interfaces/IDMXProtocolDevice.h"
+#include "DMXProtocolConstants.h"
+#include "DMXProtocolTypes.h"
+#include "DMXProtocolTransportSACN.h"
+#include "Common/UdpSocketBuilder.h"
+
+#include "DMXProtocolSACNUtils.h"
+#include "Packets/DMXProtocolE131PDUPacket.h"
+
+FDMXProtocolUniverseSACN::FDMXProtocolUniverseSACN(IDMXProtocol* InDMXProtocol, TSharedPtr<IDMXProtocolPort> InPort, uint16 InUniverseID)
+	: DMXProtocol(InDMXProtocol)
+	, Port(InPort)
+	, Priority(0)
+	, UniverseID(InUniverseID)
+{
+	checkf(DMXProtocol, TEXT("DMXProtocol pointer is nullptr"));
+	checkf(Port.IsValid(), TEXT("DMXProtocol port is not valid"));
+
+	OutputDMXBuffer = MakeShared<FDMXBuffer>();
+	InputDMXBuffer = MakeShared<FDMXBuffer>();
+
+	FIPv4Endpoint UnicastEndpoint = FIPv4Endpoint::Any;
+	TSharedPtr<FInternetAddr> MulticastAddr = FDMXProtocolSACN::GetUniverseAddr(UniverseID);
+	FIPv4Endpoint MulticastEndpoint = FIPv4Endpoint(MulticastAddr);
+
+	ListenerSocket = FUdpSocketBuilder(TEXT("UDPSACNListenerSocket"))
+		.AsNonBlocking()
+		.AsReusable()
+#if PLATFORM_WINDOWS
+		.BoundToAddress(UnicastEndpoint.Address)
+#endif
+		.BoundToPort(MulticastEndpoint.Port)
+		.JoinedToGroup(MulticastEndpoint.Address, UnicastEndpoint.Address)
+		.WithMulticastLoopback()
+		.WithMulticastTtl(1)
+		.WithMulticastInterface(UnicastEndpoint.Address)
+		.WithReceiveBufferSize(2 * 1024 * 1024);
+
+	if (ListenerSocket == nullptr)
+	{
+		UE_LOG_DMXPROTOCOL(Error, TEXT("ERROR create BroadcastSocket"));
+	}
+	SACNReceiver = MakeShared<FDMXProtocolReceiverSACN>(*ListenerSocket, FTimespan::FromMilliseconds(100));
+	SACNReceiver->OnDataReceived().BindRaw(this, &FDMXProtocolUniverseSACN::OnDataReceived);
+}
+
+FDMXProtocolUniverseSACN::~FDMXProtocolUniverseSACN()
+{
+	SACNReceiver.Reset();
+
+	// Clear all sockets
+	if (ListenerSocket)
+	{
+		ListenerSocket->Close();
+		ISocketSubsystem::Get(PLATFORM_SOCKETSUBSYSTEM)->DestroySocket(ListenerSocket);
+	}
+}
+
+IDMXProtocol * FDMXProtocolUniverseSACN::GetProtocol() const
+{
+	return DMXProtocol;
+}
+
+TWeakPtr<IDMXProtocolPort> FDMXProtocolUniverseSACN::GetCachedUniversePort() const
+{
+	return Port;
+}
+
+TSharedPtr<FDMXBuffer> FDMXProtocolUniverseSACN::GetOutputDMXBuffer() const
+{
+	return OutputDMXBuffer;
+}
+
+TSharedPtr<FDMXBuffer> FDMXProtocolUniverseSACN::GetInputDMXBuffer() const
+{
+	return InputDMXBuffer;
+}
+
+bool FDMXProtocolUniverseSACN::SetDMXFragment(const IDMXFragmentMap & DMXFragment)
+{
+	OutputDMXBuffer->SetDMXFragment(DMXFragment);
+	return false;
+}
+
+uint8 FDMXProtocolUniverseSACN::GetPriority() const
+{
+	return Priority;
+}
+
+uint16 FDMXProtocolUniverseSACN::GetUniverseID() const
+{
+	return UniverseID;
+}
+
+void FDMXProtocolUniverseSACN::OnDataReceived(const FArrayReaderPtr & Buffer)
+{
+	// It will be more handlers
+	switch (SACN::GetRootPacketType(Buffer))
+	{
+	case VECTOR_ROOT_E131_DATA:
+		HandleReplyPacket(Buffer);
+		break;
+	default:
+		break;
+	}
+}
+
+bool FDMXProtocolUniverseSACN::HandleReplyPacket(const FArrayReaderPtr & Buffer)
+{
+	// SACN PDU packets
+	FDMXProtocolE131RootLayerPacket RootLayer;
+	FDMXProtocolE131FramingLayerPacket FramingLayer;
+	FDMXProtocolE131DMPLayerPacket DMPLayer;
+
+	*Buffer << RootLayer;
+	*Buffer << FramingLayer;
+	*Buffer << DMPLayer;
+
+	// Make sure we copy same amount of data
+	if (InputDMXBuffer->GetDMXData().Num() == ACN_DMX_SIZE)
+	{
+		FMemory::Memcpy(InputDMXBuffer->GetDMXData().GetData(), DMPLayer.DMX, ACN_DMX_SIZE);
+		return true;
+	}
+	else
+	{
+		UE_LOG_DMXPROTOCOL(Warning, TEXT("Size of incoming DMX buffer is wrong"));
+		return false;
+	}
+}
diff --git a/Engine/Plugins/Experimental/DigitalMultiplex/DMXProtocol/Source/DMXProtocolSACN/Private/DMXProtocolUniverseSACN.h b/Engine/Plugins/Experimental/DigitalMultiplex/DMXProtocol/Source/DMXProtocolSACN/Private/DMXProtocolUniverseSACN.h
new file mode 100644
--- /dev/null
+++ b/Engine/Plugins/Experimental/DigitalMultiplex/DMXProtocol/Source/DMXProtocolSACN/Private/DMXProtocolUniverseSACN.h
@@ -0,0 +1,42 @@
+// Copyright 1998-2019 Epic Games, Inc. All Rights Reserved.
+
+#pragma once
+
+#include "CoreMinimal.h"
+#include "Interfaces/IPv4/IPv4Endpoint.h"
+#include "Interfaces/IDMXProtocolUniverse.h"
+
+class FSocket;
+
+class DMXPROTOCOLSACN_API FDMXProtocolUniverseSACN
+	: public IDMXProtocolUniverse
+{
+public:
+	FDMXProtocolUniverseSACN(IDMXProtocol* InDMXProtocol, TSharedPtr<IDMXProtocolPort> InPort, uint16 InUniverseID);
+	virtual ~FDMXProtocolUniverseSACN();
+
+	//~ Begin IDMXProtocolDevice implementation
+	virtual IDMXProtocol* GetProtocol() const override;
+	virtual TWeakPtr<IDMXProtocolPort> GetCachedUniversePort() const override;
+	virtual TSharedPtr<FDMXBuffer> GetOutputDMXBuffer() const override;
+	virtual TSharedPtr<FDMXBuffer> GetInputDMXBuffer() const override;
+	virtual bool SetDMXFragment(const IDMXFragmentMap& DMXFragment) override;
+	virtual uint8 GetPriority() const override;
+	virtual uint16 GetUniverseID() const override;
+	//~ End IDMXProtocolDevice implementation
+
+	void OnDataReceived(const FArrayReaderPtr & Buffer);
+
+	bool HandleReplyPacket(const FArrayReaderPtr & Buffer);
+
+private:
+	IDMXProtocol* DMXProtocol;
+	TWeakPtr<IDMXProtocolPort> Port;
+	TSharedPtr<FDMXBuffer> OutputDMXBuffer;
+	TSharedPtr<FDMXBuffer> InputDMXBuffer;
+	uint8 Priority;
+	uint16 UniverseID;
+
+	FSocket* ListenerSocket;
+	TSharedPtr<IDMXProtocolReceiver> SACNReceiver;
+};
\ No newline at end of file
diff --git a/Engine/Plugins/Experimental/DigitalMultiplex/DMXProtocol/Source/DMXProtocolSACN/Private/Packets/DMXProtocolE131PDUPacket.cpp b/Engine/Plugins/Experimental/DigitalMultiplex/DMXProtocol/Source/DMXProtocolSACN/Private/Packets/DMXProtocolE131PDUPacket.cpp
new file mode 100644
--- /dev/null
+++ b/Engine/Plugins/Experimental/DigitalMultiplex/DMXProtocol/Source/DMXProtocolSACN/Private/Packets/DMXProtocolE131PDUPacket.cpp
@@ -0,0 +1,123 @@
+#include "Packets/DMXProtocolE131PDUPacket.h"
+#include "DMXProtocolTypes.h"
+#include "Serialization/DMXMemoryWriter.h"
+
+#include "DMXProtocolMacros.h"
+
+REGISTER_DMX_ARCHIVE(FDMXProtocolE131RootLayerPacket);
+REGISTER_DMX_ARCHIVE(FDMXProtocolE131FramingLayerPacket);
+REGISTER_DMX_ARCHIVE(FDMXProtocolE131DMPLayerPacket);
+REGISTER_DMX_ARCHIVE(FDMXProtocolUDPE131FramingLayerPacket);
+REGISTER_DMX_ARCHIVE(FDMXProtocolUDPE131DiscoveryLayerPacket);
+
+TSharedPtr<FBufferArchive> FDMXProtocolE131RootLayerPacket::Pack()
+{
+	TSharedPtr<FDMXMemoryWriter> Writer = MakeShared<FDMXMemoryWriter>();
+	*Writer << *this;
+
+	return Writer;
+}
+
+void FDMXProtocolE131RootLayerPacket::Serialize(FArchive & Ar)
+{
+	Ar.SetByteSwapping(true);
+	Ar << PreambleSize;
+	Ar << PostambleSize;
+	Ar.SetByteSwapping(false);
+	Ar.Serialize((void*)ACNPacketIdentifier, ACN_IDENTIFIER_SIZE);
+	Ar.SetByteSwapping(true);
+	Ar << Flags;
+	Ar << Length;
+	Ar << Vector;
+	Ar.SetByteSwapping(false);
+	Ar.Serialize((void*)CID, ACN_CIDBYTES);
+}
+
+
+TSharedPtr<FBufferArchive> FDMXProtocolE131FramingLayerPacket::Pack()
+{
+	TSharedPtr<FDMXMemoryWriter> Writer = MakeShared<FDMXMemoryWriter>();
+	*Writer << *this;
+
+	return Writer;
+}
+
+void FDMXProtocolE131FramingLayerPacket::Serialize(FArchive & Ar)
+{
+	Ar.SetByteSwapping(true);
+	Ar << Flags;
+	Ar << Length;
+	Ar << Vector;
+	Ar.SetByteSwapping(false);
+	Ar.Serialize((void*)SourceName, ACN_SOURCE_NAME_SIZE);
+	Ar.SetByteSwapping(true);
+	Ar << Priority;
+	Ar << SynchronizationAddress;
+	Ar << SequenceNumber;
+	Ar << Options;
+	Ar << Universe;
+}
+
+
+TSharedPtr<FBufferArchive> FDMXProtocolE131DMPLayerPacket::Pack()
+{
+	TSharedPtr<FDMXMemoryWriter> Writer = MakeShared<FDMXMemoryWriter>();
+	*Writer << *this;
+
+	return Writer;
+}
+
+void FDMXProtocolE131DMPLayerPacket::Serialize(FArchive & Ar)
+{
+	Ar.SetByteSwapping(true);
+	Ar << Flags;
+	Ar << Length;
+	Ar << Vector;
+	Ar << AddressTypeAndDataType;
+	Ar << FirstPropertyAddress;
+	Ar << AddressIncrement;
+	Ar << PropertyValueCount;
+	Ar << STARTCode;
+	Ar.SetByteSwapping(false);
+	Ar.Serialize((void*)DMX, ACN_DMX_SIZE);
+}
+
+
+TSharedPtr<FBufferArchive> FDMXProtocolUDPE131FramingLayerPacket::Pack()
+{
+	TSharedPtr<FDMXMemoryWriter> Writer = MakeShared<FDMXMemoryWriter>();
+	*Writer << *this;
+
+	return Writer;
+}
+
+void FDMXProtocolUDPE131FramingLayerPacket::Serialize(FArchive & Ar)
+{
+	Ar.SetByteSwapping(true);
+	Ar << Flags;
+	Ar << Length;
+	Ar << Vector;
+	Ar.SetByteSwapping(false);
+	Ar.Serialize((void*)SourceName, ACN_SOURCE_NAME_SIZE);
+	Ar.Serialize((void*)Reserved, sizeof(Reserved));
+}
+
+TSharedPtr<FBufferArchive> FDMXProtocolUDPE131DiscoveryLayerPacket::Pack()
+{
+	TSharedPtr<FDMXMemoryWriter> Writer = MakeShared<FDMXMemoryWriter>();
+	*Writer << *this;
+
+	return Writer;
+}
+
+void FDMXProtocolUDPE131DiscoveryLayerPacket::Serialize(FArchive & Ar)
+{
+	Ar.SetByteSwapping(true);
+	Ar << Flags;
+	Ar << Length;
+	Ar << Vector;
+	Ar << Page;
+	Ar << Last;
+	Ar.SetByteSwapping(false);
+	Ar.Serialize((void*)Universes, sizeof(Universes));
+}
diff --git a/Engine/Plugins/Experimental/DigitalMultiplex/DMXProtocol/Source/DMXProtocolSACN/Private/Packets/DMXProtocolE131PDUPacket.h b/Engine/Plugins/Experimental/DigitalMultiplex/DMXProtocol/Source/DMXProtocolSACN/Private/Packets/DMXProtocolE131PDUPacket.h
new file mode 100644
--- /dev/null
+++ b/Engine/Plugins/Experimental/DigitalMultiplex/DMXProtocol/Source/DMXProtocolSACN/Private/Packets/DMXProtocolE131PDUPacket.h
@@ -0,0 +1,126 @@
+// Copyright 1998-2019 Epic Games, Inc. All Rights Reserved.
+
+#pragma once
+
+#include "CoreMinimal.h"
+#include "Interfaces/IDMXProtocolPacket.h"
+#include "DMXProtocolSACNConstants.h"
+
+struct FDMXProtocolE131RootLayerPacket
+	: public IDMXProtocolPacket
+{
+public:
+	FDMXProtocolE131RootLayerPacket()
+	{
+		FMemory::Memcpy(ACNPacketIdentifier, FSACNConstants::ACN_IDENTIFIER, ACN_IDENTIFIER_SIZE);
+	}
+
+	//~ Begin IDMXProtocolPacket implementation
+	virtual TSharedPtr<FBufferArchive> Pack() override;
+	//~ End IDMXProtocolPort implementation
+
+	void Serialize(FArchive& Ar);
+	friend FArchive& operator<<(FArchive& Ar, FDMXProtocolE131RootLayerPacket& Packet);
+
+public:
+	uint16 PreambleSize = ACN_RLP_PREAMBLE_SIZE;
+	uint16 PostambleSize = 0;
+	uint8 ACNPacketIdentifier[ACN_IDENTIFIER_SIZE] = { 0 };
+	uint8 Flags = 0x72;
+	uint8 Length = 0x6e;
+	uint32 Vector = VECTOR_ROOT_E131_DATA;
+	uint8 CID[ACN_CIDBYTES] = { 0 };
+};
+
+struct FDMXProtocolE131FramingLayerPacket
+	: public IDMXProtocolPacket
+{
+public:
+	FDMXProtocolE131FramingLayerPacket()
+	{
+		FMemory::Memcpy(SourceName, "New Source", ACN_SOURCE_NAME_SIZE);
+	}
+
+	//~ Begin IDMXProtocolPacket implementation
+	virtual TSharedPtr<FBufferArchive> Pack() override;
+	//~ End IDMXProtocolPort implementation
+
+	void Serialize(FArchive& Ar);
+	friend FArchive& operator<<(FArchive& Ar, FDMXProtocolE131FramingLayerPacket& Packet);
+
+public:
+	uint8 Flags = 0;
+	uint8 Length = 0;
+	uint32 Vector = VECTOR_E131_DATA_PACKET;
+	uint8 SourceName[ACN_SOURCE_NAME_SIZE] = { 0 };
+	uint8 Priority = 0;
+	uint16 SynchronizationAddress = 0;
+	uint8 SequenceNumber = 0;
+	uint8 Options = 0;
+	uint8 Universe = 0;
+};
+
+struct FDMXProtocolE131DMPLayerPacket
+	: public IDMXProtocolPacket
+{
+public:
+	//~ Begin IDMXProtocolPacket implementation
+	virtual TSharedPtr<FBufferArchive> Pack() override;
+	//~ End IDMXProtocolPort implementation
+
+	void Serialize(FArchive& Ar);
+	friend FArchive& operator<<(FArchive& Ar, FDMXProtocolE131DMPLayerPacket& Packet);
+
+public:
+	uint8 Flags = 0;
+	uint8 Length = 0;
+	uint8 Vector = VECTOR_DMP_SET_PROPERTY;
+	uint8 AddressTypeAndDataType = 0;
+	uint16 FirstPropertyAddress = 0;
+	uint16 AddressIncrement = 0;
+	uint16 PropertyValueCount = 0;
+	uint8 STARTCode = 0;
+	uint8 DMX[ACN_DMX_SIZE] = { 0 };
+};
+
+struct FDMXProtocolUDPE131FramingLayerPacket
+	: public IDMXProtocolPacket
+{
+public:
+	//~ Begin IDMXProtocolPacket implementation
+	virtual TSharedPtr<FBufferArchive> Pack() override;
+	//~ End IDMXProtocolPort implementation
+
+	void Serialize(FArchive& Ar);
+	friend FArchive& operator<<(FArchive& Ar, FDMXProtocolUDPE131FramingLayerPacket& Packet);
+
+public:
+	uint8 Flags = 0x72;
+	uint8 Length = 0x0b;
+	uint32 Vector = 2;
+	uint8 SourceName[ACN_SOURCE_NAME_SIZE] = { 0 };
+	uint8 Reserved[4] = { 0 };
+
+};
+
+struct FDMXProtocolUDPE131DiscoveryLayerPacket
+	: public IDMXProtocolPacket
+{
+public:
+	//~ Begin IDMXProtocolPacket implementation
+	virtual TSharedPtr<FBufferArchive> Pack() override;
+	//~ End IDMXProtocolPort implementation
+
+	void Serialize(FArchive& Ar);
+	friend FArchive& operator<<(FArchive& Ar, FDMXProtocolUDPE131DiscoveryLayerPacket& Packet);
+
+public:
+	uint8 Flags = 0;
+	uint8 Length = 0;
+	uint32 Vector = VECTOR_E131_EXTENDED_DISCOVERY;
+	uint8 Page = 0;
+	uint8 Last = 0;
+	uint16 Universes[ACN_DMX_SIZE] = { 0 };
+
+};
+
diff --git a/Engine/Plugins/Experimental/DigitalMultiplex/DMXProtocol/Source/DMXProtocolSACN/Private/Tests/DMXPrtotocolSACNFlowTest.cpp b/Engine/Plugins/Experimental/DigitalMultiplex/DMXProtocol/Source/DMXProtocolSACN/Private/Tests/DMXPrtotocolSACNFlowTest.cpp
new file mode 100644
--- /dev/null
+++ b/Engine/Plugins/Experimental/DigitalMultiplex/DMXProtocol/Source/DMXProtocolSACN/Private/Tests/DMXPrtotocolSACNFlowTest.cpp
@@ -0,0 +1,160 @@
+// Copyright 1998-2019 Epic Games, Inc. All Rights Reserved.
+
+#include "CoreMinimal.h"
+#include "Engine.h"
+#include "DMXProtocolSACNModule.h"
+
+#include "DMXProtocolSACN.h"
+
+#include "DMXProtocolSACN.h"
+#include "DMXProtocolPortSACN.h"
+
+#include "Managers/DMXProtocolDeviceManager.h"
+#include "Managers/DMXProtocolInterfaceManager.h"
+#include "Managers/DMXProtocolPortManager.h"
+#include "Managers/DMXProtocolUniverseManager.h"
+
+#include "DMXProtocolInterfaceSACN.h"
+#include "DMXProtocolDeviceSACN.h"
+#include "DMXProtocolUniverseSACN.h"
+#include "DMXProtocolPortSACN.h"
+#include "DMXProtocolSACNConstants.h"
+
+struct FDMXProtocolLauncher
+{
+	FDMXProtocolLauncher(uint8 InDeviceID, uint8 InPortID)
+		: DeviceID(InDeviceID)
+		, PortID(InPortID)
+		, InterfaceID(0)
+		, UniverseID(3)
+	{
+		DMXProtocol = IDMXProtocol::Get<FDMXProtocolSACN>(FDMXProtocolSACNModule::NAME_SACN);
+
+		InterfaceSACN = MakeShared<FDMXProtocolInterfaceSACN>(DMXProtocol);
+		DMXProtocol->GetInterfaceManager()->AddInterface(0, InterfaceSACN);
+
+		FJsonObject DeviceSettings;
+		ProtocolDevice = MakeShared<FDMXProtocolDeviceSACN>(DMXProtocol, InterfaceSACN, DeviceSettings, DeviceID);
+		DMXProtocol->GetDeviceManager()->AddDevice(ProtocolDevice);
+
+		ProtocolOutputPort = MakeShared<FDMXProtocolPortSACN>(DMXProtocol, ProtocolDevice, PortID, UniverseID, DMX_PORT_OUTPUT);
+		DMXProtocol->GetPortManager()->AddOutputPort(ProtocolDevice, ProtocolOutputPort);
+
+		FJsonObject UniverseSettings;
+		ProtocolUniverse = MakeShared<FDMXProtocolUniverseSACN>(DMXProtocol, ProtocolOutputPort, UniverseID);
+		DMXProtocol->GetUniverseManager()->AddUniverse(ProtocolOutputPort, ProtocolUniverse);
+		ProtocolOutputPort->SetUniverse(ProtocolUniverse);
+	}
+
+	~FDMXProtocolLauncher()
+	{
+		DMXProtocol->GetInterfaceManager()->RemoveInterface(InterfaceID);
+		DMXProtocol->GetDeviceManager()->RemoveDevice(ProtocolDevice);
+		DMXProtocol->GetPortManager()->RemoveOutputPort(ProtocolDevice, ProtocolOutputPort);
+		DMXProtocol->GetUniverseManager()->RemoveUniverse(ProtocolOutputPort, ProtocolUniverse);
+	}
+
+	const uint8 DeviceID;
+	const uint8 PortID;
+	const uint8 InterfaceID;
+	const uint16 UniverseID;
+	TSharedPtr<FDMXProtocolInterfaceSACN> InterfaceSACN;
+	TSharedPtr<FDMXProtocolDeviceSACN> ProtocolDevice;
+	TSharedPtr<FDMXProtocolPortSACN> ProtocolOutputPort;
+	TSharedPtr<FDMXProtocolUniverseSACN> ProtocolUniverse;
+	FDMXProtocolSACN* DMXProtocol;
+};
+
+
+IMPLEMENT_SIMPLE_AUTOMATION_TEST(DMXPrtotocolSACNBasicFlowTest, "VirtualProduction.DMX.SACN.BasicFlow", EAutomationTestFlags::ApplicationContextMask | EAutomationTestFlags::EngineFilter)
+
+static void SensDMXFragment(DMXPrtotocolSACNBasicFlowTest& Test, FDMXProtocolSACN* DMXProtocol, TSharedPtr<FDMXProtocolDeviceSACN> InProtocolDevice, uint8 InPortID)
+{
+	IDMXProtocolPort* SACNPort = DMXProtocol->GetPortManager()->GetPortByDeviceAndID(InProtocolDevice, InPortID, EDMXPortDirection::DMX_PORT_OUTPUT);
+	Test.TestNotNull(TEXT("SACNPort is valid"), SACNPort);
+	Test.TestEqual(TEXT("Universe ID should be equal 3"), SACNPort->GetUniverseID(), 3);
+
+
+	// Send DMX fragment 1
+	const uint8 FixtureChannels[6] = { 0, 1, 2, 5, 6, 7 };
+	const uint8 FixtureValues[6] = { 255, 155, 50, 100, 200, 220 };
+	IDMXFragmentMap DMXFragmentMap;
+	DMXFragmentMap.Add(FixtureChannels[0], FixtureValues[0]);
+	DMXFragmentMap.Add(FixtureChannels[1], FixtureValues[1]);
+	DMXFragmentMap.Add(FixtureChannels[2], FixtureValues[2]);
+	DMXProtocol->SetDMXFragment(SACNPort->GetUniverseID(), DMXFragmentMap);
+
+	// Wait before recieve DMX
+	FPlatformProcess::Sleep(0.2f);
+	TSharedPtr<IDMXProtocolUniverse> Universe = SACNPort->GetCachedUniverse().Pin();
+	Test.TestTrue(TEXT("Universe should be valid"), Universe.IsValid());
+	TSharedPtr<FDMXBuffer> InputDMXBuffer = Universe->GetInputDMXBuffer();
+
+	Test.TestEqual(TEXT("Incoming buffer should be same"), FixtureValues[0], InputDMXBuffer->GetDMXData()[FixtureChannels[0]]);
+	Test.TestEqual(TEXT("Incoming buffer should be same"), FixtureValues[1], InputDMXBuffer->GetDMXData()[FixtureChannels[1]]);
+	Test.TestEqual(TEXT("Incoming buffer should be same"), FixtureValues[2], InputDMXBuffer->GetDMXData()[FixtureChannels[2]]);
+
+	// Send DMX fragment 2
+	DMXFragmentMap.Add(FixtureChannels[3], FixtureValues[3]);
+	DMXFragmentMap.Add(FixtureChannels[4], FixtureValues[4]);
+	DMXFragmentMap.Add(FixtureChannels[5], FixtureValues[5]);
+	DMXProtocol->SetDMXFragment(SACNPort->GetUniverseID(), DMXFragmentMap);
+	FPlatformProcess::Sleep(0.2f);
+
+	Test.TestEqual(TEXT("Old value DMX input should be the same"), FixtureValues[0], InputDMXBuffer->GetDMXData()[FixtureChannels[0]]);
+	Test.TestEqual(TEXT("Incoming buffer should be same"), FixtureValues[4], InputDMXBuffer->GetDMXData()[FixtureChannels[4]]);
+	Test.TestEqual(TEXT("Incoming buffer should be same"), FixtureValues[5], InputDMXBuffer->GetDMXData()[FixtureChannels[5]]);
+
+	// @TODO Implement test cases for Discovery
+	TArray<uint16> Universes;
+	for (int UniversesIndex = 0; UniversesIndex <= DMX_UNIVERSE_SIZE; UniversesIndex++)
+	{
+		Universes.Add(UniversesIndex);
+	}
+	DMXProtocol->SendDiscovery(Universes);
+}
+
+bool DMXPrtotocolSACNBasicFlowTest::RunTest(const FString& Parameters)
+{
+	FDMXProtocolLauncher DMXProtocolLauncher(0, 0);
+
+	// Test Sending DMX and Discovery
+	SensDMXFragment(*this, DMXProtocolLauncher.DMXProtocol, DMXProtocolLauncher.ProtocolDevice, DMXProtocolLauncher.PortID);
+
+	return true;
+}
+
+IMPLEMENT_SIMPLE_AUTOMATION_TEST(FDMXPrtotocolSACNConsoleCommandsTest, "VirtualProduction.DMX.SACN.ConsoleCommands", EAutomationTestFlags::ApplicationContextMask | EAutomationTestFlags::EngineFilter)
+
+bool FDMXPrtotocolSACNConsoleCommandsTest::RunTest(const FString& Parameters)
+{
+	FDMXProtocolLauncher DMXProtocolLauncher(0, 0);
+
+	// DMX.SACN.SendDMX[UniverseID] Channel:Value Channel : Value Channel : Value ...
+	uint32 Values[3][2] =
+	{
+		{95, 105},
+		{88, 41},
+		{12, 1},
+	};
+
+	FString Command(
+		FString::Printf(TEXT("DMX.SACN.SendDMX %d %d:%d %d:%d %d:%d"),
+			DMXProtocolLauncher.ProtocolOutputPort->GetUniverseID(),
+			Values[0][0], Values[0][1],
+			Values[1][0], Values[1][1],
+			Values[2][0], Values[2][1]
+		));
+
+	GEngine->Exec(GWorld, *Command);
+
+	FPlatformProcess::Sleep(0.2f);
+
+	TSharedPtr<IDMXProtocolUniverse> Universe = DMXProtocolLauncher.ProtocolOutputPort->GetCachedUniverse().Pin();
+	TSharedPtr<FDMXBuffer> InputDMXBuffer = Universe->GetInputDMXBuffer();
+	TestEqual(TEXT("Incoming buffer should be same"), Values[0][1], InputDMXBuffer->GetDMXData()[Values[0][0]]);
+	TestEqual(TEXT("Incoming buffer should be same"), Values[1][1], InputDMXBuffer->GetDMXData()[Values[1][0]]);
+	TestEqual(TEXT("Incoming buffer should be same"), Values[2][1], InputDMXBuffer->GetDMXData()[Values[2][0]]);
+
+	return true;
+}
\ No newline at end of file
diff --git a/Engine/Plugins/Experimental/DigitalMultiplex/DMXProtocol/Source/DMXProtocolSACN/Public/DMXProtocolSACN.h b/Engine/Plugins/Experimental/DigitalMultiplex/DMXProtocol/Source/DMXProtocolSACN/Public/DMXProtocolSACN.h
new file mode 100644
--- /dev/null
+++ b/Engine/Plugins/Experimental/DigitalMultiplex/DMXProtocol/Source/DMXProtocolSACN/Public/DMXProtocolSACN.h
@@ -0,0 +1,55 @@
+// Copyright 1998-2019 Epic Games, Inc. All Rights Reserved.
+
+#pragma once
+
+#include "CoreMinimal.h"
+#include "HAL/ThreadSafeCounter.h"
+#include "DMXProtocol.h"
+#include "DMXProtocolSACNModule.h"
+#include "Interfaces/IPv4/IPv4Endpoint.h"
+
+class FInternetAddr;
+class FSocket;
+
+class DMXPROTOCOLSACN_API FDMXProtocolSACN
+	: public FDMXProtocol
+{
+public:
+	//~ Begin IDMXProtocolBase implementation
+	virtual bool Init() override;
+	virtual bool Shutdown() override;
+	virtual bool Tick(float DeltaTime) override;
+	//~ End IDMXProtocolBase implementation
+
+	//~ Begin IDMXProtocol implementation
+	virtual TSharedPtr<IDMXProtocolSender> GetSenderInterface() const override;
+	virtual FName GetProtocolName() const override;
+	virtual void Reload() override;
+	virtual bool SendDMX(uint16 UniverseID, uint8 PortID, const TSharedPtr<FDMXBuffer>& DMXBuffer) const override;
+	virtual bool IsEnabled() const override;
+	//~ End IDMXProtocol implementation
+
+	//~ Begin IDMXProtocolRDM implementation
+	virtual void SendRDMCommand(const TSharedPtr<FJsonObject>& CMD) override;
+	virtual void RDMDiscovery(const TSharedPtr<FJsonObject>& CMD) override;
+	//~ End IDMXProtocol implementation
+
+	//~ sACN specific implementation
+	bool SendDiscovery(const TArray<uint16>& Universes);
+
+	//~ Only the factory makes instances
+	FDMXProtocolSACN() = delete;
+	explicit FDMXProtocolSACN(FName InProtocolName, FJsonObject& InSettings)
+		: FDMXProtocol(FDMXProtocolSACNModule::NAME_SACN, InSettings)
+		, SenderSocket(nullptr)
+	{}
+
+public:
+	static TSharedPtr<FInternetAddr> GetUniverseAddr(uint16 InUniverseID);
+
+private:
+	TSharedPtr<IDMXProtocolSender> SACNSender;
+
+	/** Holds the network socket used to sender packages. */
+	FSocket* SenderSocket;
+};
diff --git a/Engine/Plugins/Experimental/DigitalMultiplex/DMXProtocol/Source/DMXProtocolSACN/Public/DMXProtocolSACNModule.h b/Engine/Plugins/Experimental/DigitalMultiplex/DMXProtocol/Source/DMXProtocolSACN/Public/DMXProtocolSACNModule.h
new file mode 100644
--- /dev/null
+++ b/Engine/Plugins/Experimental/DigitalMultiplex/DMXProtocol/Source/DMXProtocolSACN/Public/DMXProtocolSACNModule.h
@@ -0,0 +1,53 @@
+// Copyright 1998-2019 Epic Games, Inc. All Rights Reserved.
+
+#pragma once
+
+#include "CoreMinimal.h"
+#include "Modules/ModuleInterface.h"
+#include "IDMXProtocolFactory.h"
+
+/**
+ */
+class FDMXProtocolFactorySACN : public IDMXProtocolFactory
+{
+public:
+	virtual TSharedPtr<IDMXProtocol> CreateProtocol(const FName& ProtocolName) override;
+};
+
+/**
+ */
+class FDMXProtocolSACNModule : public IModuleInterface
+{
+private:
+
+	/** Class responsible for creating instance(s) of the protocol */
+	TUniquePtr<FDMXProtocolFactorySACN> FactorySACN;
+
+public:
+	static FName const NAME_SACN;
+
+public:
+	//~ Begin IModuleInterface implementation
+	virtual void StartupModule() override;
+	virtual void ShutdownModule() override;
+	//~ End IModuleInterface implementation
+
+	FDMXProtocolSACNModule& Get();
+
+private:
+	/**
+	 * Sending DMX through console command
+	 * Command structure is DMX.SACN.SendDMX [UniverseID] Channel:Value Channel:Value Channel:Value ...
+	 * Example:
+	 * DMX.SACN.SendDMX 7 25:156 26:0 27:10 28:15
+	 * It will send the DMX to Universe 7. It could be any value from 0 to 64214
+	 * And it update the channels values for channel 25, 26, 27, 28 with values 156, 0, 10, 15
+	 * Channel could be any value from 0 to 511
+	 * Value could be any value from 0 to 255
+	 */
+	static void SendDMXCommandHandler(const TArray<FString>& Args);
+
+private:
+	/** Command for sending DMX through the console */
+	static FAutoConsoleCommand SendDMXCommand;
+};
