Fix Android build for DMX plugin.
Remove const string. Copy string in the Packet constructor

diff --git a/Engine/Plugins/Experimental/DigitalMultiplex/DMXProtocol/Source/DMXProtocolArtNet/Private/DMXProtocolArtNetConstants.h b/Engine/Plugins/Experimental/DigitalMultiplex/DMXProtocol/Source/DMXProtocolArtNet/Private/DMXProtocolArtNetConstants.h
--- a/Engine/Plugins/Experimental/DigitalMultiplex/DMXProtocol/Source/DMXProtocolArtNet/Private/DMXProtocolArtNetConstants.h
+++ b/Engine/Plugins/Experimental/DigitalMultiplex/DMXProtocol/Source/DMXProtocolArtNet/Private/DMXProtocolArtNetConstants.h
@@ -133,8 +133,3 @@ enum
 	RECV_NO_DATA = 1,
 	MAX_NODE_BCAST_LIMIT = 30 // always bcast after this point
 };
-
-struct DMXArtNet
-{
-	static uint8 constexpr ARTNET_STRING[ARTNET_STRING_SIZE] = "Art-Net";
-};
diff --git a/Engine/Plugins/Experimental/DigitalMultiplex/DMXProtocol/Source/DMXProtocolArtNet/Private/Packets/DMXProtocolArtNetPackets.h b/Engine/Plugins/Experimental/DigitalMultiplex/DMXProtocol/Source/DMXProtocolArtNet/Private/Packets/DMXProtocolArtNetPackets.h
--- a/Engine/Plugins/Experimental/DigitalMultiplex/DMXProtocol/Source/DMXProtocolArtNet/Private/Packets/DMXProtocolArtNetPackets.h
+++ b/Engine/Plugins/Experimental/DigitalMultiplex/DMXProtocol/Source/DMXProtocolArtNet/Private/Packets/DMXProtocolArtNetPackets.h
@@ -12,7 +12,7 @@ struct FDMXProtocolArtNetPacket
 public:
 	FDMXProtocolArtNetPacket()
 	{
-		FMemory::Memcpy(ID, DMXArtNet::ARTNET_STRING, ARTNET_STRING_SIZE);
+		FMemory::Memcpy(ID, "Art-Net\0", ARTNET_STRING_SIZE);
 	}
 
 public:
diff --git a/Engine/Plugins/Experimental/DigitalMultiplex/DMXProtocol/Source/DMXProtocolSACN/Private/DMXProtocolSACNConstants.h b/Engine/Plugins/Experimental/DigitalMultiplex/DMXProtocol/Source/DMXProtocolSACN/Private/DMXProtocolSACNConstants.h
--- a/Engine/Plugins/Experimental/DigitalMultiplex/DMXProtocol/Source/DMXProtocolSACN/Private/DMXProtocolSACNConstants.h
+++ b/Engine/Plugins/Experimental/DigitalMultiplex/DMXProtocol/Source/DMXProtocolSACN/Private/DMXProtocolSACNConstants.h
@@ -56,9 +56,3 @@ enum
 	VHD_H_FLAG = 0x20,
 	VHD_D_FLAG = 0x10
 };
-
-struct FSACNConstants
-{
-	static uint8 constexpr ACN_IDENTIFIER[ACN_IDENTIFIER_SIZE] = { 'A', 'S', 'C', '-', 'E', '1', '.', '1', '7', '\0', '\0', '\0' };
-};
-
diff --git a/Engine/Plugins/Experimental/DigitalMultiplex/DMXProtocol/Source/DMXProtocolSACN/Private/Packets/DMXProtocolE131PDUPacket.h b/Engine/Plugins/Experimental/DigitalMultiplex/DMXProtocol/Source/DMXProtocolSACN/Private/Packets/DMXProtocolE131PDUPacket.h
--- a/Engine/Plugins/Experimental/DigitalMultiplex/DMXProtocol/Source/DMXProtocolSACN/Private/Packets/DMXProtocolE131PDUPacket.h
+++ b/Engine/Plugins/Experimental/DigitalMultiplex/DMXProtocol/Source/DMXProtocolSACN/Private/Packets/DMXProtocolE131PDUPacket.h
@@ -12,7 +12,7 @@ struct FDMXProtocolE131RootLayerPacket
 public:
 	FDMXProtocolE131RootLayerPacket()
 	{
-		FMemory::Memcpy(ACNPacketIdentifier, FSACNConstants::ACN_IDENTIFIER, ACN_IDENTIFIER_SIZE);
+		FMemory::Memcpy(ACNPacketIdentifier, "ASC-E1.17\0\0\0:", ACN_IDENTIFIER_SIZE);
 	}
 
 	//~ Begin IDMXProtocolPacket implementation
