Fix includes for VS2019. DMX Prorocol Plugin

diff --git a/Engine/Plugins/Experimental/DigitalMultiplex/DMXProtocol/Source/DMXProtocol/Private/Tests/DMXProtocolTests.cpp b/Engine/Plugins/Experimental/DigitalMultiplex/DMXProtocol/Source/DMXProtocol/Private/Tests/DMXProtocolTests.cpp
--- a/Engine/Plugins/Experimental/DigitalMultiplex/DMXProtocol/Source/DMXProtocol/Private/Tests/DMXProtocolTests.cpp
+++ b/Engine/Plugins/Experimental/DigitalMultiplex/DMXProtocol/Source/DMXProtocol/Private/Tests/DMXProtocolTests.cpp
@@ -2,7 +2,7 @@
 
 #include "CoreMinimal.h"
 #include "Misc/AutomationTest.h"
-#include "IDMXProtocolFactory.h"
+#include "Interfaces/IDMXProtocolFactory.h"
 #include "DMXProtocolCommon.h"
 #include "Json.h"
 
diff --git a/Engine/Plugins/Experimental/DigitalMultiplex/DMXProtocol/Source/DMXProtocol/Public/DMXProtocolTypes.h b/Engine/Plugins/Experimental/DigitalMultiplex/DMXProtocol/Source/DMXProtocol/Public/DMXProtocolTypes.h
--- a/Engine/Plugins/Experimental/DigitalMultiplex/DMXProtocol/Source/DMXProtocol/Public/DMXProtocolTypes.h
+++ b/Engine/Plugins/Experimental/DigitalMultiplex/DMXProtocol/Source/DMXProtocol/Public/DMXProtocolTypes.h
@@ -2,7 +2,7 @@
 
 #pragma once
 
-#include "IDMXProtocol.h"
+#include "Interfaces/IDMXProtocol.h"
 #include "DMXProtocolCommon.h"
 #include "DMXProtocolConstants.h"
 #include "Json.h"
diff --git a/Engine/Plugins/Experimental/DigitalMultiplex/DMXProtocol/Source/DMXProtocolArtNet/Public/DMXProtocolArtNetModule.h b/Engine/Plugins/Experimental/DigitalMultiplex/DMXProtocol/Source/DMXProtocolArtNet/Public/DMXProtocolArtNetModule.h
--- a/Engine/Plugins/Experimental/DigitalMultiplex/DMXProtocol/Source/DMXProtocolArtNet/Public/DMXProtocolArtNetModule.h
+++ b/Engine/Plugins/Experimental/DigitalMultiplex/DMXProtocol/Source/DMXProtocolArtNet/Public/DMXProtocolArtNetModule.h
@@ -4,7 +4,7 @@
 
 #include "CoreMinimal.h"
 #include "Modules/ModuleInterface.h"
-#include "IDMXProtocolFactory.h"
+#include "Interfaces/IDMXProtocolFactory.h"
 
 /**
  */
diff --git a/Engine/Plugins/Experimental/DigitalMultiplex/DMXProtocol/Source/DMXProtocolSACN/Public/DMXProtocolSACNModule.h b/Engine/Plugins/Experimental/DigitalMultiplex/DMXProtocol/Source/DMXProtocolSACN/Public/DMXProtocolSACNModule.h
--- a/Engine/Plugins/Experimental/DigitalMultiplex/DMXProtocol/Source/DMXProtocolSACN/Public/DMXProtocolSACNModule.h
+++ b/Engine/Plugins/Experimental/DigitalMultiplex/DMXProtocol/Source/DMXProtocolSACN/Public/DMXProtocolSACNModule.h
@@ -4,7 +4,7 @@
 
 #include "CoreMinimal.h"
 #include "Modules/ModuleInterface.h"
-#include "IDMXProtocolFactory.h"
+#include "Interfaces/IDMXProtocolFactory.h"
 
 /**
  */
