Another DMX-related include fix for VS2019

diff --git a/Engine/Plugins/Experimental/DigitalMultiplex/DMXProtocol/Source/DMXProtocolArtNet/Public/DMXProtocolArtNetModule.h b/Engine/Plugins/Experimental/DigitalMultiplex/DMXProtocol/Source/DMXProtocolArtNet/Public/DMXProtocolArtNetModule.h
--- a/Engine/Plugins/Experimental/DigitalMultiplex/DMXProtocol/Source/DMXProtocolArtNet/Public/DMXProtocolArtNetModule.h
+++ b/Engine/Plugins/Experimental/DigitalMultiplex/DMXProtocol/Source/DMXProtocolArtNet/Public/DMXProtocolArtNetModule.h
@@ -4,6 +4,7 @@
 
 #include "CoreMinimal.h"
 #include "Modules/ModuleInterface.h"
+#include "HAL/IConsoleManager.h"
 #include "Interfaces/IDMXProtocolFactory.h"
 
 /**
