Another DMX-related include fix for VS2019 SACN console manager include

diff --git a/Engine/Plugins/Experimental/DigitalMultiplex/DMXProtocol/Source/DMXProtocolSACN/Public/DMXProtocolSACNModule.h b/Engine/Plugins/Experimental/DigitalMultiplex/DMXProtocol/Source/DMXProtocolSACN/Public/DMXProtocolSACNModule.h
--- a/Engine/Plugins/Experimental/DigitalMultiplex/DMXProtocol/Source/DMXProtocolSACN/Public/DMXProtocolSACNModule.h
+++ b/Engine/Plugins/Experimental/DigitalMultiplex/DMXProtocol/Source/DMXProtocolSACN/Public/DMXProtocolSACNModule.h
@@ -3,6 +3,7 @@
 #pragma once
 
 #include "CoreMinimal.h"
+#include "HAL/IConsoleManager.h"
 #include "Modules/ModuleInterface.h"
 #include "Interfaces/IDMXProtocolFactory.h"
 
