Fix DMX Build
Fix headers class forward declaration
Fix SACN and ArtNet header

diff --git a/Engine/Plugins/Experimental/DigitalMultiplex/DMXProtocol/Source/DMXProtocolArtNet/Private/DMXProtocolTransportArtNet.cpp b/Engine/Plugins/Experimental/DigitalMultiplex/DMXProtocol/Source/DMXProtocolArtNet/Private/DMXProtocolTransportArtNet.cpp
--- a/Engine/Plugins/Experimental/DigitalMultiplex/DMXProtocol/Source/DMXProtocolArtNet/Private/DMXProtocolTransportArtNet.cpp
+++ b/Engine/Plugins/Experimental/DigitalMultiplex/DMXProtocol/Source/DMXProtocolArtNet/Private/DMXProtocolTransportArtNet.cpp
@@ -11,6 +11,7 @@
 #include "Serialization/ArrayReader.h"
 #include "Serialization/ArrayWriter.h"
 #include "Sockets.h"
+#include "SocketSubsystem.h"
 
 static FTimespan CalculateWaitTime()
 {
diff --git a/Engine/Plugins/Experimental/DigitalMultiplex/DMXProtocol/Source/DMXProtocolArtNet/Private/DMXProtocolTransportArtNet.h b/Engine/Plugins/Experimental/DigitalMultiplex/DMXProtocol/Source/DMXProtocolArtNet/Private/DMXProtocolTransportArtNet.h
--- a/Engine/Plugins/Experimental/DigitalMultiplex/DMXProtocol/Source/DMXProtocolArtNet/Private/DMXProtocolTransportArtNet.h
+++ b/Engine/Plugins/Experimental/DigitalMultiplex/DMXProtocol/Source/DMXProtocolArtNet/Private/DMXProtocolTransportArtNet.h
@@ -7,11 +7,10 @@
 #include "Interfaces/IDMXProtocolTransport.h"
 #include "DMXProtocolTypes.h"
 
-
 class FSocket;
+class ISocketSubsystem;
 class FDMXProtocolArtNet;
 
-
 class DMXPROTOCOLARTNET_API FDMXProtocolSenderArtNet
 	: public IDMXProtocolSender
 {
diff --git a/Engine/Plugins/Experimental/DigitalMultiplex/DMXProtocol/Source/DMXProtocolArtNet/Private/Packets/DMXProtocolArtNetPackets.h b/Engine/Plugins/Experimental/DigitalMultiplex/DMXProtocol/Source/DMXProtocolArtNet/Private/Packets/DMXProtocolArtNetPackets.h
--- a/Engine/Plugins/Experimental/DigitalMultiplex/DMXProtocol/Source/DMXProtocolArtNet/Private/Packets/DMXProtocolArtNetPackets.h
+++ b/Engine/Plugins/Experimental/DigitalMultiplex/DMXProtocol/Source/DMXProtocolArtNet/Private/Packets/DMXProtocolArtNetPackets.h
@@ -10,13 +10,7 @@ struct FDMXProtocolArtNetPacket
 	: public IDMXProtocolPacket
 {
 public:
-	FDMXProtocolArtNetPacket()
-	{
-		FMemory::Memcpy(ID, "Art-Net\0", ARTNET_STRING_SIZE);
-	}
-
-public:
-	uint8 ID[8] = { 0 };
+	uint8 ID[ARTNET_STRING_SIZE] = "Art-Net";
 };
 
 struct FDMXProtocolArtNetDMXPacket
diff --git a/Engine/Plugins/Experimental/DigitalMultiplex/DMXProtocol/Source/DMXProtocolSACN/Private/DMXProtocolTransportSACN.cpp b/Engine/Plugins/Experimental/DigitalMultiplex/DMXProtocol/Source/DMXProtocolSACN/Private/DMXProtocolTransportSACN.cpp
--- a/Engine/Plugins/Experimental/DigitalMultiplex/DMXProtocol/Source/DMXProtocolSACN/Private/DMXProtocolTransportSACN.cpp
+++ b/Engine/Plugins/Experimental/DigitalMultiplex/DMXProtocol/Source/DMXProtocolSACN/Private/DMXProtocolTransportSACN.cpp
@@ -11,6 +11,7 @@
 #include "Serialization/ArrayReader.h"
 #include "Serialization/ArrayWriter.h"
 #include "Sockets.h"
+#include "SocketSubsystem.h"
 
 static FTimespan CalculateWaitTime()
 {
diff --git a/Engine/Plugins/Experimental/DigitalMultiplex/DMXProtocol/Source/DMXProtocolSACN/Private/DMXProtocolTransportSACN.h b/Engine/Plugins/Experimental/DigitalMultiplex/DMXProtocol/Source/DMXProtocolSACN/Private/DMXProtocolTransportSACN.h
--- a/Engine/Plugins/Experimental/DigitalMultiplex/DMXProtocol/Source/DMXProtocolSACN/Private/DMXProtocolTransportSACN.h
+++ b/Engine/Plugins/Experimental/DigitalMultiplex/DMXProtocol/Source/DMXProtocolSACN/Private/DMXProtocolTransportSACN.h
@@ -7,11 +7,10 @@
 #include "Interfaces/IDMXProtocolTransport.h"
 #include "DMXProtocolTypes.h"
 
-
 class FSocket;
+class ISocketSubsystem;
 class FDMXProtocolSACN;
 
-
 class DMXPROTOCOLSACN_API FDMXProtocolSenderSACN
 	: public IDMXProtocolSender
 {
diff --git a/Engine/Plugins/Experimental/DigitalMultiplex/DMXProtocol/Source/DMXProtocolSACN/Private/DMXProtocolUniverseSACN.cpp b/Engine/Plugins/Experimental/DigitalMultiplex/DMXProtocol/Source/DMXProtocolSACN/Private/DMXProtocolUniverseSACN.cpp
--- a/Engine/Plugins/Experimental/DigitalMultiplex/DMXProtocol/Source/DMXProtocolSACN/Private/DMXProtocolUniverseSACN.cpp
+++ b/Engine/Plugins/Experimental/DigitalMultiplex/DMXProtocol/Source/DMXProtocolSACN/Private/DMXProtocolUniverseSACN.cpp
@@ -8,7 +8,6 @@
 #include "Common/UdpSocketBuilder.h"
 
 #include "DMXProtocolSACNUtils.h"
-#include "Packets/DMXProtocolE131PDUPacket.h"
 
 FDMXProtocolUniverseSACN::FDMXProtocolUniverseSACN(IDMXProtocol* InDMXProtocol, TSharedPtr<IDMXProtocolPort> InPort, uint16 InUniverseID)
 	: DMXProtocol(InDMXProtocol)
@@ -111,18 +110,14 @@ void FDMXProtocolUniverseSACN::OnDataReceived(const FArrayReaderPtr & Buffer)
 bool FDMXProtocolUniverseSACN::HandleReplyPacket(const FArrayReaderPtr & Buffer)
 {
 	// SACN PDU packets
-	FDMXProtocolE131RootLayerPacket RootLayer;
-	FDMXProtocolE131FramingLayerPacket FramingLayer;
-	FDMXProtocolE131DMPLayerPacket DMPLayer;
-
-	*Buffer << RootLayer;
-	*Buffer << FramingLayer;
-	*Buffer << DMPLayer;
+	*Buffer << IncomingDMXRootLayer;
+	*Buffer << IncomingDMXFramingLayer;
+	*Buffer << IncomingDMXDMPLayer;
 
 	// Make sure we copy same amount of data
 	if (InputDMXBuffer->GetDMXData().Num() == ACN_DMX_SIZE)
 	{
-		FMemory::Memcpy(InputDMXBuffer->GetDMXData().GetData(), DMPLayer.DMX, ACN_DMX_SIZE);
+		FMemory::Memcpy(InputDMXBuffer->GetDMXData().GetData(), IncomingDMXDMPLayer.DMX, ACN_DMX_SIZE);
 		return true;
 	}
 	else
diff --git a/Engine/Plugins/Experimental/DigitalMultiplex/DMXProtocol/Source/DMXProtocolSACN/Private/DMXProtocolUniverseSACN.h b/Engine/Plugins/Experimental/DigitalMultiplex/DMXProtocol/Source/DMXProtocolSACN/Private/DMXProtocolUniverseSACN.h
--- a/Engine/Plugins/Experimental/DigitalMultiplex/DMXProtocol/Source/DMXProtocolSACN/Private/DMXProtocolUniverseSACN.h
+++ b/Engine/Plugins/Experimental/DigitalMultiplex/DMXProtocol/Source/DMXProtocolSACN/Private/DMXProtocolUniverseSACN.h
@@ -5,6 +5,7 @@
 #include "CoreMinimal.h"
 #include "Interfaces/IPv4/IPv4Endpoint.h"
 #include "Interfaces/IDMXProtocolUniverse.h"
+#include "Packets/DMXProtocolE131PDUPacket.h"
 
 class FSocket;
 
@@ -25,6 +26,12 @@ public:
 	virtual uint16 GetUniverseID() const override;
 	//~ End IDMXProtocolDevice implementation
 
+	//~ SACN Getters functions
+	const FDMXProtocolE131RootLayerPacket& GetIncomingDMXRootLayer() const { return IncomingDMXRootLayer; }
+	const FDMXProtocolE131FramingLayerPacket& GetIncomingDMXFramingLayer() const { return IncomingDMXFramingLayer; }
+	const FDMXProtocolE131DMPLayerPacket& GetIncomingDMXDMPLayer() const { return IncomingDMXDMPLayer; }
+
+private:
 	void OnDataReceived(const FArrayReaderPtr & Buffer);
 
 	bool HandleReplyPacket(const FArrayReaderPtr & Buffer);
@@ -39,4 +46,8 @@ private:
 
 	FSocket* ListenerSocket;
 	TSharedPtr<IDMXProtocolReceiver> SACNReceiver;
+
+	FDMXProtocolE131RootLayerPacket IncomingDMXRootLayer;
+	FDMXProtocolE131FramingLayerPacket IncomingDMXFramingLayer;
+	FDMXProtocolE131DMPLayerPacket IncomingDMXDMPLayer;
 };
\ No newline at end of file
diff --git a/Engine/Plugins/Experimental/DigitalMultiplex/DMXProtocol/Source/DMXProtocolSACN/Private/Packets/DMXProtocolE131PDUPacket.h b/Engine/Plugins/Experimental/DigitalMultiplex/DMXProtocol/Source/DMXProtocolSACN/Private/Packets/DMXProtocolE131PDUPacket.h
--- a/Engine/Plugins/Experimental/DigitalMultiplex/DMXProtocol/Source/DMXProtocolSACN/Private/Packets/DMXProtocolE131PDUPacket.h
+++ b/Engine/Plugins/Experimental/DigitalMultiplex/DMXProtocol/Source/DMXProtocolSACN/Private/Packets/DMXProtocolE131PDUPacket.h
@@ -10,11 +10,6 @@ struct FDMXProtocolE131RootLayerPacket
 	: public IDMXProtocolPacket
 {
 public:
-	FDMXProtocolE131RootLayerPacket()
-	{
-		FMemory::Memcpy(ACNPacketIdentifier, "ASC-E1.17\0\0\0:", ACN_IDENTIFIER_SIZE);
-	}
-
 	//~ Begin IDMXProtocolPacket implementation
 	virtual TSharedPtr<FBufferArchive> Pack() override;
 	//~ End IDMXProtocolPort implementation
@@ -25,7 +20,7 @@ public:
 public:
 	uint16 PreambleSize = ACN_RLP_PREAMBLE_SIZE;
 	uint16 PostambleSize = 0;
-	uint8 ACNPacketIdentifier[ACN_IDENTIFIER_SIZE] = { 0 };
+	uint8 ACNPacketIdentifier[ACN_IDENTIFIER_SIZE] = { 'A', 'S', 'C', '-', 'E', '1', '.', '1', '7', '\0', '\0', '\0' };
 	uint8 Flags = 0x72;
 	uint8 Length = 0x6e;
 	uint32 Vector = VECTOR_ROOT_E131_DATA;
@@ -36,11 +31,6 @@ struct FDMXProtocolE131FramingLayerPacket
 	: public IDMXProtocolPacket
 {
 public:
-	FDMXProtocolE131FramingLayerPacket()
-	{
-		FMemory::Memcpy(SourceName, "New Source", ACN_SOURCE_NAME_SIZE);
-	}
-
 	//~ Begin IDMXProtocolPacket implementation
 	virtual TSharedPtr<FBufferArchive> Pack() override;
 	//~ End IDMXProtocolPort implementation
@@ -52,7 +42,7 @@ public:
 	uint8 Flags = 0;
 	uint8 Length = 0;
 	uint32 Vector = VECTOR_E131_DATA_PACKET;
-	uint8 SourceName[ACN_SOURCE_NAME_SIZE] = { 0 };
+	uint8 SourceName[ACN_SOURCE_NAME_SIZE] = "New Source";
 	uint8 Priority = 0;
 	uint16 SynchronizationAddress = 0;
 	uint8 SequenceNumber = 0;
diff --git a/Engine/Plugins/Experimental/DigitalMultiplex/DMXProtocol/Source/DMXProtocolSACN/Private/Tests/DMXPrtotocolSACNFlowTest.cpp b/Engine/Plugins/Experimental/DigitalMultiplex/DMXProtocol/Source/DMXProtocolSACN/Private/Tests/DMXPrtotocolSACNFlowTest.cpp
--- a/Engine/Plugins/Experimental/DigitalMultiplex/DMXProtocol/Source/DMXProtocolSACN/Private/Tests/DMXPrtotocolSACNFlowTest.cpp
+++ b/Engine/Plugins/Experimental/DigitalMultiplex/DMXProtocol/Source/DMXProtocolSACN/Private/Tests/DMXPrtotocolSACNFlowTest.cpp
@@ -20,6 +20,7 @@
 #include "DMXProtocolPortSACN.h"
 #include "DMXProtocolSACNConstants.h"
 
+
 struct FDMXProtocolLauncher
 {
 	FDMXProtocolLauncher(uint8 InDeviceID, uint8 InPortID)
@@ -90,6 +91,15 @@ static void SensDMXFragment(DMXPrtotocolSACNBasicFlowTest& Test, FDMXProtocolSAC
 	Test.TestTrue(TEXT("Universe should be valid"), Universe.IsValid());
 	TSharedPtr<FDMXBuffer> InputDMXBuffer = Universe->GetInputDMXBuffer();
 
+	// Test package ACNPacketIdentifier
+	{
+		const uint8 ACNPacketIdentifier[ACN_IDENTIFIER_SIZE] = { 'A', 'S', 'C', '-', 'E', '1', '.', '1', '7', '\0', '\0', '\0' };
+		FDMXProtocolUniverseSACN* DMXProtocolUniverseSACN = static_cast<FDMXProtocolUniverseSACN*>(Universe.Get());
+		FString ExpectedInentifier = FString::Printf(TEXT("%s"), ACNPacketIdentifier);
+		FString IncomingInentifier = FString::Printf(TEXT("%s"), DMXProtocolUniverseSACN->GetIncomingDMXRootLayer().ACNPacketIdentifier);
+		Test.TestTrue(TEXT("Incoming ACNPacketIdentifier should be the same"), ExpectedInentifier.Equals(IncomingInentifier));
+	}
+
 	Test.TestEqual(TEXT("Incoming buffer should be same"), FixtureValues[0], InputDMXBuffer->GetDMXData()[FixtureChannels[0]]);
 	Test.TestEqual(TEXT("Incoming buffer should be same"), FixtureValues[1], InputDMXBuffer->GetDMXData()[FixtureChannels[1]]);
 	Test.TestEqual(TEXT("Incoming buffer should be same"), FixtureValues[2], InputDMXBuffer->GetDMXData()[FixtureChannels[2]]);
