Fix build warnings for DMX Plugin

diff --git a/Engine/Plugins/Experimental/DigitalMultiplex/DMXProtocol/Source/DMXProtocol/Private/DMXProtocol.cpp b/Engine/Plugins/Experimental/DigitalMultiplex/DMXProtocol/Source/DMXProtocol/Private/DMXProtocol.cpp
--- a/Engine/Plugins/Experimental/DigitalMultiplex/DMXProtocol/Source/DMXProtocol/Private/DMXProtocol.cpp
+++ b/Engine/Plugins/Experimental/DigitalMultiplex/DMXProtocol/Source/DMXProtocol/Private/DMXProtocol.cpp
@@ -11,7 +11,7 @@
 #include "Interfaces/IDMXProtocolUniverse.h"
 
 #include "DMXProtocolTypes.h"
-#include "Json.h"
+#include "Dom/JsonObject.h"
 
 #include "HAL/IConsoleManager.h"
 #include "Logging/LogMacros.h"
diff --git a/Engine/Plugins/Experimental/DigitalMultiplex/DMXProtocol/Source/DMXProtocol/Private/DMXProtocolModule.cpp b/Engine/Plugins/Experimental/DigitalMultiplex/DMXProtocol/Source/DMXProtocol/Private/DMXProtocolModule.cpp
--- a/Engine/Plugins/Experimental/DigitalMultiplex/DMXProtocol/Source/DMXProtocol/Private/DMXProtocolModule.cpp
+++ b/Engine/Plugins/Experimental/DigitalMultiplex/DMXProtocol/Source/DMXProtocol/Private/DMXProtocolModule.cpp
@@ -105,6 +105,11 @@ IDMXProtocol* FDMXProtocolModule::GetProtocol(const FName InProtocolName)
 	return (DMXProtocolPtr == nullptr) ? nullptr : (*DMXProtocolPtr).Get();
 }
 
+const TMap<FName, IDMXProtocolFactory*>& FDMXProtocolModule::GetProtocolFactories() const
+{
+	return DMXFactories;
+}
+
 void FDMXProtocolModule::ShutdownDMXProtocol(const FName& ProtocolName)
 {
 	if (!ProtocolName.IsNone())
diff --git a/Engine/Plugins/Experimental/DigitalMultiplex/DMXProtocol/Source/DMXProtocol/Private/Tests/DMXProtocolTests.cpp b/Engine/Plugins/Experimental/DigitalMultiplex/DMXProtocol/Source/DMXProtocol/Private/Tests/DMXProtocolTests.cpp
--- a/Engine/Plugins/Experimental/DigitalMultiplex/DMXProtocol/Source/DMXProtocol/Private/Tests/DMXProtocolTests.cpp
+++ b/Engine/Plugins/Experimental/DigitalMultiplex/DMXProtocol/Source/DMXProtocol/Private/Tests/DMXProtocolTests.cpp
@@ -4,7 +4,7 @@
 #include "Misc/AutomationTest.h"
 #include "Interfaces/IDMXProtocolFactory.h"
 #include "DMXProtocolCommon.h"
-#include "Json.h"
+#include "Dom/JsonObject.h"
 
 #include "DMXProtocol.h"
 
diff --git a/Engine/Plugins/Experimental/DigitalMultiplex/DMXProtocol/Source/DMXProtocol/Public/DMXProtocolModule.h b/Engine/Plugins/Experimental/DigitalMultiplex/DMXProtocol/Source/DMXProtocol/Public/DMXProtocolModule.h
--- a/Engine/Plugins/Experimental/DigitalMultiplex/DMXProtocol/Source/DMXProtocol/Public/DMXProtocolModule.h
+++ b/Engine/Plugins/Experimental/DigitalMultiplex/DMXProtocol/Source/DMXProtocol/Public/DMXProtocolModule.h
@@ -22,6 +22,8 @@ public:
 	 * @return Return the pointer to protocol.
 	 */
 	virtual IDMXProtocol* GetProtocol(const FName InProtocolName = NAME_None);
+	
+	const TMap<FName, IDMXProtocolFactory*>& GetProtocolFactories() const;
 
 	//~ Begin IModuleInterface implementation
 	virtual void ShutdownModule() override;
diff --git a/Engine/Plugins/Experimental/DigitalMultiplex/DMXProtocol/Source/DMXProtocol/Public/DMXProtocolTypes.h b/Engine/Plugins/Experimental/DigitalMultiplex/DMXProtocol/Source/DMXProtocol/Public/DMXProtocolTypes.h
--- a/Engine/Plugins/Experimental/DigitalMultiplex/DMXProtocol/Source/DMXProtocol/Public/DMXProtocolTypes.h
+++ b/Engine/Plugins/Experimental/DigitalMultiplex/DMXProtocol/Source/DMXProtocol/Public/DMXProtocolTypes.h
@@ -5,7 +5,7 @@
 #include "Interfaces/IDMXProtocol.h"
 #include "DMXProtocolCommon.h"
 #include "DMXProtocolConstants.h"
-#include "Json.h"
+#include "Dom/JsonObject.h"
 
 struct DMXPROTOCOL_API FDMXBuffer
 {
diff --git a/Engine/Plugins/Experimental/DigitalMultiplex/DMXProtocol/Source/DMXProtocol/Public/Interfaces/IDMXProtocol.h b/Engine/Plugins/Experimental/DigitalMultiplex/DMXProtocol/Source/DMXProtocol/Public/Interfaces/IDMXProtocol.h
--- a/Engine/Plugins/Experimental/DigitalMultiplex/DMXProtocol/Source/DMXProtocol/Public/Interfaces/IDMXProtocol.h
+++ b/Engine/Plugins/Experimental/DigitalMultiplex/DMXProtocol/Source/DMXProtocol/Public/Interfaces/IDMXProtocol.h
@@ -57,6 +57,13 @@ public:
 		return static_cast<TProtocol*>(DMXProtocolModule.GetProtocol(ProtocolName));
 	}
 
+	static const TMap<FName, IDMXProtocolFactory*>& GetProtocolFactories()
+	{
+		static const FName DMXProtocolModuleName = TEXT("DMXProtocol");
+		FDMXProtocolModule& DMXProtocolModule = FModuleManager::GetModuleChecked<FDMXProtocolModule>(DMXProtocolModuleName);
+		return DMXProtocolModule.GetProtocolFactories();
+	}
+
 	/**
 	 * If protocol exists return the pointer otherwise it create a new protocol first and then return the pointer.
 	 * @param  InProtocolName Name of the requested protocol
diff --git a/Engine/Plugins/Experimental/DigitalMultiplex/DMXProtocol/Source/DMXProtocolArtNet/Private/DMXProtocolArtNetModule.cpp b/Engine/Plugins/Experimental/DigitalMultiplex/DMXProtocol/Source/DMXProtocolArtNet/Private/DMXProtocolArtNetModule.cpp
--- a/Engine/Plugins/Experimental/DigitalMultiplex/DMXProtocol/Source/DMXProtocolArtNet/Private/DMXProtocolArtNetModule.cpp
+++ b/Engine/Plugins/Experimental/DigitalMultiplex/DMXProtocol/Source/DMXProtocolArtNet/Private/DMXProtocolArtNetModule.cpp
@@ -3,7 +3,7 @@
 #include "DMXProtocolArtNetModule.h"
 #include "CoreMinimal.h"
 #include "Modules/ModuleManager.h"
-#include "Json.h"
+#include "Dom/JsonObject.h"
 
 #include "DMXProtocolArtNet.h"
 #include "DMXProtocolTypes.h"
@@ -11,7 +11,7 @@
 
 IMPLEMENT_MODULE(FDMXProtocolArtNetModule, DMXProtocolArtNet);
 
-const FName FDMXProtocolArtNetModule::NAME_Artnet = FName(TEXT("ARTNET"));
+const FName FDMXProtocolArtNetModule::NAME_Artnet = FName(TEXT("Art-Net"));
 
 FAutoConsoleCommand FDMXProtocolArtNetModule::SendDMXCommand(
 	TEXT("DMX.ArtNet.SendDMX"),
diff --git a/Engine/Plugins/Experimental/DigitalMultiplex/DMXProtocol/Source/DMXProtocolArtNet/Private/DMXProtocolArtNetUtils.h b/Engine/Plugins/Experimental/DigitalMultiplex/DMXProtocol/Source/DMXProtocolArtNet/Private/DMXProtocolArtNetUtils.h
--- a/Engine/Plugins/Experimental/DigitalMultiplex/DMXProtocol/Source/DMXProtocolArtNet/Private/DMXProtocolArtNetUtils.h
+++ b/Engine/Plugins/Experimental/DigitalMultiplex/DMXProtocol/Source/DMXProtocolArtNet/Private/DMXProtocolArtNetUtils.h
@@ -6,6 +6,8 @@
 #include "DMXProtocolTypes.h"
 #include "DMXProtocolArtNetConstants.h"
 
+#include "Serialization/ArrayReader.h"
+
 namespace ArtNet
 {
 	static uint16 GetPacketType(const FArrayReaderPtr& Buffer)
diff --git a/Engine/Plugins/Experimental/DigitalMultiplex/DMXProtocol/Source/DMXProtocolArtNet/Private/DMXProtocolDeviceArtNet.cpp b/Engine/Plugins/Experimental/DigitalMultiplex/DMXProtocol/Source/DMXProtocolArtNet/Private/DMXProtocolDeviceArtNet.cpp
--- a/Engine/Plugins/Experimental/DigitalMultiplex/DMXProtocol/Source/DMXProtocolArtNet/Private/DMXProtocolDeviceArtNet.cpp
+++ b/Engine/Plugins/Experimental/DigitalMultiplex/DMXProtocol/Source/DMXProtocolArtNet/Private/DMXProtocolDeviceArtNet.cpp
@@ -4,7 +4,7 @@
 #include "Interfaces/IDMXProtocolPort.h"
 #include "Managers/DMXProtocolPortManager.h"
 #include "DMXProtocolTypes.h"
-#include "Json.h"
+#include "Dom/JsonObject.h"
 
 FDMXProtocolDeviceArtNet::FDMXProtocolDeviceArtNet(IDMXProtocol* InDMXProtocol, TSharedPtr<IDMXProtocolInterface> InProtocolInterface, FJsonObject& InSettings, uint32 InDeviceID)
 	: DMXProtocol(InDMXProtocol)
diff --git a/Engine/Plugins/Experimental/DigitalMultiplex/DMXProtocol/Source/DMXProtocolArtNet/Private/DMXProtocolDeviceArtNet.h b/Engine/Plugins/Experimental/DigitalMultiplex/DMXProtocol/Source/DMXProtocolArtNet/Private/DMXProtocolDeviceArtNet.h
--- a/Engine/Plugins/Experimental/DigitalMultiplex/DMXProtocol/Source/DMXProtocolArtNet/Private/DMXProtocolDeviceArtNet.h
+++ b/Engine/Plugins/Experimental/DigitalMultiplex/DMXProtocol/Source/DMXProtocolArtNet/Private/DMXProtocolDeviceArtNet.h
@@ -4,7 +4,7 @@
 
 #include "CoreMinimal.h"
 #include "Interfaces/IDMXProtocolDevice.h"
-#include "Json.h"
+#include "Dom/JsonObject.h"
 
 class DMXPROTOCOLARTNET_API FDMXProtocolDeviceArtNet
 	: public IDMXProtocolDevice
diff --git a/Engine/Plugins/Experimental/DigitalMultiplex/DMXProtocol/Source/DMXProtocolArtNet/Private/DMXProtocolPortArtNet.cpp b/Engine/Plugins/Experimental/DigitalMultiplex/DMXProtocol/Source/DMXProtocolArtNet/Private/DMXProtocolPortArtNet.cpp
--- a/Engine/Plugins/Experimental/DigitalMultiplex/DMXProtocol/Source/DMXProtocolArtNet/Private/DMXProtocolPortArtNet.cpp
+++ b/Engine/Plugins/Experimental/DigitalMultiplex/DMXProtocol/Source/DMXProtocolArtNet/Private/DMXProtocolPortArtNet.cpp
@@ -2,7 +2,7 @@
 #include "Interfaces/IDMXProtocol.h"
 #include "Interfaces/IDMXProtocolUniverse.h"
 
-#include "Json.h"
+#include "Dom/JsonObject.h"
 
 //~ Artnet specific port constants
 static const uint8 HIGH_NIBBLE = 0xF0;
diff --git a/Engine/Plugins/Experimental/DigitalMultiplex/DMXProtocol/Source/DMXProtocolArtNet/Private/DMXProtocolUniverseArtNet.cpp b/Engine/Plugins/Experimental/DigitalMultiplex/DMXProtocol/Source/DMXProtocolArtNet/Private/DMXProtocolUniverseArtNet.cpp
--- a/Engine/Plugins/Experimental/DigitalMultiplex/DMXProtocol/Source/DMXProtocolArtNet/Private/DMXProtocolUniverseArtNet.cpp
+++ b/Engine/Plugins/Experimental/DigitalMultiplex/DMXProtocol/Source/DMXProtocolArtNet/Private/DMXProtocolUniverseArtNet.cpp
@@ -1,5 +1,5 @@
 #include "DMXProtocolUniverseArtNet.h"
-#include "Json.h"
+#include "Dom/JsonObject.h"
 
 #include "Interfaces/IDMXProtocol.h"
 #include "Interfaces/IDMXProtocolPort.h"
diff --git a/Engine/Plugins/Experimental/DigitalMultiplex/DMXProtocol/Source/DMXProtocolArtNet/Private/Tests/DMXPrtotocolArtNetFlowTest.cpp b/Engine/Plugins/Experimental/DigitalMultiplex/DMXProtocol/Source/DMXProtocolArtNet/Private/Tests/DMXPrtotocolArtNetFlowTest.cpp
--- a/Engine/Plugins/Experimental/DigitalMultiplex/DMXProtocol/Source/DMXProtocolArtNet/Private/Tests/DMXPrtotocolArtNetFlowTest.cpp
+++ b/Engine/Plugins/Experimental/DigitalMultiplex/DMXProtocol/Source/DMXProtocolArtNet/Private/Tests/DMXPrtotocolArtNetFlowTest.cpp
@@ -1,7 +1,7 @@
 // Copyright 1998-2019 Epic Games, Inc. All Rights Reserved.
 
 #include "CoreMinimal.h"
-#include "Engine.h"
+#include "Engine/Engine.h"
 #include "DMXProtocolArtNetModule.h"
 
 #include "DMXProtocolArtNet.h"
diff --git a/Engine/Plugins/Experimental/DigitalMultiplex/DMXProtocol/Source/DMXProtocolSACN/Private/DMXProtocolDeviceSACN.cpp b/Engine/Plugins/Experimental/DigitalMultiplex/DMXProtocol/Source/DMXProtocolSACN/Private/DMXProtocolDeviceSACN.cpp
--- a/Engine/Plugins/Experimental/DigitalMultiplex/DMXProtocol/Source/DMXProtocolSACN/Private/DMXProtocolDeviceSACN.cpp
+++ b/Engine/Plugins/Experimental/DigitalMultiplex/DMXProtocol/Source/DMXProtocolSACN/Private/DMXProtocolDeviceSACN.cpp
@@ -3,7 +3,7 @@
 #include "Interfaces/IDMXProtocolInterface.h"
 #include "Interfaces/IDMXProtocolPort.h"
 #include "DMXProtocolTypes.h"
-#include "Json.h"
+#include "Dom/JsonObject.h"
 
 FDMXProtocolDeviceSACN::FDMXProtocolDeviceSACN(IDMXProtocol* InDMXProtocol, TSharedPtr<IDMXProtocolInterface> InProtocolInterface, FJsonObject& InSettings, uint32 InDeviceID)
 	: DMXProtocol(InDMXProtocol)
diff --git a/Engine/Plugins/Experimental/DigitalMultiplex/DMXProtocol/Source/DMXProtocolSACN/Private/DMXProtocolDeviceSACN.h b/Engine/Plugins/Experimental/DigitalMultiplex/DMXProtocol/Source/DMXProtocolSACN/Private/DMXProtocolDeviceSACN.h
--- a/Engine/Plugins/Experimental/DigitalMultiplex/DMXProtocol/Source/DMXProtocolSACN/Private/DMXProtocolDeviceSACN.h
+++ b/Engine/Plugins/Experimental/DigitalMultiplex/DMXProtocol/Source/DMXProtocolSACN/Private/DMXProtocolDeviceSACN.h
@@ -4,7 +4,7 @@
 
 #include "CoreMinimal.h"
 #include "Interfaces/IDMXProtocolDevice.h"
-#include "Json.h"
+#include "Dom/JsonObject.h"
 
 class DMXPROTOCOLSACN_API FDMXProtocolDeviceSACN
 	: public IDMXProtocolDevice
diff --git a/Engine/Plugins/Experimental/DigitalMultiplex/DMXProtocol/Source/DMXProtocolSACN/Private/DMXProtocolSACNModule.cpp b/Engine/Plugins/Experimental/DigitalMultiplex/DMXProtocol/Source/DMXProtocolSACN/Private/DMXProtocolSACNModule.cpp
--- a/Engine/Plugins/Experimental/DigitalMultiplex/DMXProtocol/Source/DMXProtocolSACN/Private/DMXProtocolSACNModule.cpp
+++ b/Engine/Plugins/Experimental/DigitalMultiplex/DMXProtocol/Source/DMXProtocolSACN/Private/DMXProtocolSACNModule.cpp
@@ -7,7 +7,7 @@
 #include "Interfaces/IDMXProtocol.h"
 #include "DMXProtocolSACN.h"
 #include "DMXProtocolTypes.h"
-#include "Json.h"
+#include "Dom/JsonObject.h"
 
 #include "Managers/DMXProtocolDeviceManager.h"
 #include "Managers/DMXProtocolInterfaceManager.h"
@@ -22,7 +22,7 @@
 
 IMPLEMENT_MODULE(FDMXProtocolSACNModule, DMXProtocolSACN);
 
-const FName FDMXProtocolSACNModule::NAME_SACN = FName(TEXT("SACN"));
+const FName FDMXProtocolSACNModule::NAME_SACN = FName(TEXT("sACN"));
 
 FAutoConsoleCommand FDMXProtocolSACNModule::SendDMXCommand(
 	TEXT("DMX.SACN.SendDMX"),
diff --git a/Engine/Plugins/Experimental/DigitalMultiplex/DMXProtocol/Source/DMXProtocolSACN/Private/DMXProtocolSACNUtils.h b/Engine/Plugins/Experimental/DigitalMultiplex/DMXProtocol/Source/DMXProtocolSACN/Private/DMXProtocolSACNUtils.h
--- a/Engine/Plugins/Experimental/DigitalMultiplex/DMXProtocol/Source/DMXProtocolSACN/Private/DMXProtocolSACNUtils.h
+++ b/Engine/Plugins/Experimental/DigitalMultiplex/DMXProtocol/Source/DMXProtocolSACN/Private/DMXProtocolSACNUtils.h
@@ -6,6 +6,8 @@
 #include "DMXProtocolTypes.h"
 #include "DMXProtocolSACNConstants.h"
 
+#include "Serialization/ArrayReader.h"
+
 namespace SACN
 {
 	static uint32 GetRootPacketType(const FArrayReaderPtr& Buffer)
diff --git a/Engine/Plugins/Experimental/DigitalMultiplex/DMXProtocol/Source/DMXProtocolSACN/Private/Tests/DMXPrtotocolSACNFlowTest.cpp b/Engine/Plugins/Experimental/DigitalMultiplex/DMXProtocol/Source/DMXProtocolSACN/Private/Tests/DMXPrtotocolSACNFlowTest.cpp
--- a/Engine/Plugins/Experimental/DigitalMultiplex/DMXProtocol/Source/DMXProtocolSACN/Private/Tests/DMXPrtotocolSACNFlowTest.cpp
+++ b/Engine/Plugins/Experimental/DigitalMultiplex/DMXProtocol/Source/DMXProtocolSACN/Private/Tests/DMXPrtotocolSACNFlowTest.cpp
@@ -1,7 +1,7 @@
 // Copyright 1998-2019 Epic Games, Inc. All Rights Reserved.
 
 #include "CoreMinimal.h"
-#include "Engine.h"
+#include "Engine/Engine.h"
 #include "DMXProtocolSACNModule.h"
 
 #include "DMXProtocolSACN.h"
diff --git a/Engine/Plugins/Experimental/DigitalMultiplex/DMXProtocol/Source/DMXProtocolSACN/Public/DMXProtocolSACNModule.h b/Engine/Plugins/Experimental/DigitalMultiplex/DMXProtocol/Source/DMXProtocolSACN/Public/DMXProtocolSACNModule.h
--- a/Engine/Plugins/Experimental/DigitalMultiplex/DMXProtocol/Source/DMXProtocolSACN/Public/DMXProtocolSACNModule.h
+++ b/Engine/Plugins/Experimental/DigitalMultiplex/DMXProtocol/Source/DMXProtocolSACN/Public/DMXProtocolSACNModule.h
@@ -17,7 +17,7 @@ public:
 
 /**
  */
-class FDMXProtocolSACNModule : public IModuleInterface
+class DMXPROTOCOLSACN_API FDMXProtocolSACNModule : public IModuleInterface
 {
 private:
 
