Build warnings in DMX plugin

#UE-85196
#UE-85240
#UE-85243

diff --git a/Engine/Plugins/Experimental/DigitalMultiplex/DMXProtocol/Source/DMXProtocol/Public/Interfaces/IDMXProtocol.h b/Engine/Plugins/Experimental/DigitalMultiplex/DMXProtocol/Source/DMXProtocol/Public/Interfaces/IDMXProtocol.h
--- a/Engine/Plugins/Experimental/DigitalMultiplex/DMXProtocol/Source/DMXProtocol/Public/Interfaces/IDMXProtocol.h
+++ b/Engine/Plugins/Experimental/DigitalMultiplex/DMXProtocol/Source/DMXProtocol/Public/Interfaces/IDMXProtocol.h
@@ -11,10 +11,10 @@
 #include "Interfaces/IDMXProtocolRDM.h"
 #include "Interfaces/IDMXProtocolTransport.h"
 
-#include "Managers/DMXProtocolDeviceManager.h"
-#include "Managers/DMXProtocolInterfaceManager.h"
-#include "Managers/DMXProtocolPortManager.h"
-#include "Managers/DMXProtocolUniverseManager.h"
+class FDMXProtocolDeviceManager;
+class FDMXProtocolInterfaceManager;
+class FDMXProtocolPortManager;
+class FDMXProtocolUniverseManager;
 
 DMXPROTOCOL_API DECLARE_LOG_CATEGORY_EXTERN(LogDMXProtocol, Log, All);
 
diff --git a/Engine/Plugins/Experimental/DigitalMultiplex/DMXProtocol/Source/DMXProtocolArtNet/Private/DMXProtocolArtNet.cpp b/Engine/Plugins/Experimental/DigitalMultiplex/DMXProtocol/Source/DMXProtocolArtNet/Private/DMXProtocolArtNet.cpp
--- a/Engine/Plugins/Experimental/DigitalMultiplex/DMXProtocol/Source/DMXProtocolArtNet/Private/DMXProtocolArtNet.cpp
+++ b/Engine/Plugins/Experimental/DigitalMultiplex/DMXProtocol/Source/DMXProtocolArtNet/Private/DMXProtocolArtNet.cpp
@@ -8,6 +8,8 @@
 #include "DMXProtocolArtNetUtils.h"
 #include "DMXProtocolPortArtNet.h"
 #include "DMXProtocolPackager.h"
+#include "Managers/DMXProtocolDeviceManager.h"
+#include "Managers/DMXProtocolPortManager.h"
 
 using namespace ArtNet;
 
@@ -47,6 +49,7 @@ bool FDMXProtocolArtNet::Init()
 	if (BroadcastSocket == nullptr)
 	{
 		UE_LOG_DMXPROTOCOL(Error, TEXT("ERROR create BroadcastSocket"));
+		return false;
 	}
 
 	// Create sender
diff --git a/Engine/Plugins/Experimental/DigitalMultiplex/DMXProtocol/Source/DMXProtocolArtNet/Private/Tests/DMXPrtotocolArtNetFlowTest.cpp b/Engine/Plugins/Experimental/DigitalMultiplex/DMXProtocol/Source/DMXProtocolArtNet/Private/Tests/DMXPrtotocolArtNetFlowTest.cpp
--- a/Engine/Plugins/Experimental/DigitalMultiplex/DMXProtocol/Source/DMXProtocolArtNet/Private/Tests/DMXPrtotocolArtNetFlowTest.cpp
+++ b/Engine/Plugins/Experimental/DigitalMultiplex/DMXProtocol/Source/DMXProtocolArtNet/Private/Tests/DMXPrtotocolArtNetFlowTest.cpp
@@ -67,7 +67,7 @@ struct FDMXProtocolLauncher
 
 IMPLEMENT_SIMPLE_AUTOMATION_TEST(FDMXPrtotocolArtNetBasicFlowTest, "VirtualProduction.DMX.ArtNet.BasicFlow", EAutomationTestFlags::ApplicationContextMask | EAutomationTestFlags::EngineFilter)
 
-static void SensDMXFragment(FDMXPrtotocolArtNetBasicFlowTest& Test, const FDMXProtocolLauncher& DMXProtocolLauncher)
+static void SensDMXFragment(FAutomationTestBase& Test, const FDMXProtocolLauncher& DMXProtocolLauncher)
 {
 	FDMXProtocolPortArtNet* ArtNetPort = static_cast<FDMXProtocolPortArtNet*>(DMXProtocolLauncher.DMXProtocol->GetPortManager()->GetPortByDeviceAndID(DMXProtocolLauncher.ProtocolDevice, DMXProtocolLauncher.PortID, EDMXPortDirection::DMX_PORT_OUTPUT));
 
diff --git a/Engine/Plugins/Experimental/DigitalMultiplex/DMXProtocol/Source/DMXProtocolSACN/Private/DMXProtocolSACN.cpp b/Engine/Plugins/Experimental/DigitalMultiplex/DMXProtocol/Source/DMXProtocolSACN/Private/DMXProtocolSACN.cpp
--- a/Engine/Plugins/Experimental/DigitalMultiplex/DMXProtocol/Source/DMXProtocolSACN/Private/DMXProtocolSACN.cpp
+++ b/Engine/Plugins/Experimental/DigitalMultiplex/DMXProtocol/Source/DMXProtocolSACN/Private/DMXProtocolSACN.cpp
@@ -33,6 +33,7 @@ bool FDMXProtocolSACN::Init()
 	if (SenderSocket == nullptr)
 	{
 		UE_LOG_DMXPROTOCOL(Error, TEXT("ERROR create SenderSocket"));
+		return false;
 	}
 
 	SACNSender = MakeShared<FDMXProtocolSenderSACN>(*SenderSocket, this);
diff --git a/Engine/Plugins/Experimental/DigitalMultiplex/DMXProtocol/Source/DMXProtocolSACN/Private/DMXProtocolUniverseSACN.cpp b/Engine/Plugins/Experimental/DigitalMultiplex/DMXProtocol/Source/DMXProtocolSACN/Private/DMXProtocolUniverseSACN.cpp
--- a/Engine/Plugins/Experimental/DigitalMultiplex/DMXProtocol/Source/DMXProtocolSACN/Private/DMXProtocolUniverseSACN.cpp
+++ b/Engine/Plugins/Experimental/DigitalMultiplex/DMXProtocol/Source/DMXProtocolSACN/Private/DMXProtocolUniverseSACN.cpp
@@ -42,8 +42,11 @@ FDMXProtocolUniverseSACN::FDMXProtocolUniverseSACN(IDMXProtocol* InDMXProtocol,
 	{
 		UE_LOG_DMXPROTOCOL(Error, TEXT("ERROR create BroadcastSocket"));
 	}
-	SACNReceiver = MakeShared<FDMXProtocolReceiverSACN>(*ListenerSocket, FTimespan::FromMilliseconds(100));
-	SACNReceiver->OnDataReceived().BindRaw(this, &FDMXProtocolUniverseSACN::OnDataReceived);
+	else
+	{
+		SACNReceiver = MakeShared<FDMXProtocolReceiverSACN>(*ListenerSocket, FTimespan::FromMilliseconds(100));
+		SACNReceiver->OnDataReceived().BindRaw(this, &FDMXProtocolUniverseSACN::OnDataReceived);
+	}
 }
 
 FDMXProtocolUniverseSACN::~FDMXProtocolUniverseSACN()
diff --git a/Engine/Plugins/Experimental/DigitalMultiplex/DMXProtocol/Source/DMXProtocolSACN/Private/Tests/DMXPrtotocolSACNFlowTest.cpp b/Engine/Plugins/Experimental/DigitalMultiplex/DMXProtocol/Source/DMXProtocolSACN/Private/Tests/DMXPrtotocolSACNFlowTest.cpp
--- a/Engine/Plugins/Experimental/DigitalMultiplex/DMXProtocol/Source/DMXProtocolSACN/Private/Tests/DMXPrtotocolSACNFlowTest.cpp
+++ b/Engine/Plugins/Experimental/DigitalMultiplex/DMXProtocol/Source/DMXProtocolSACN/Private/Tests/DMXPrtotocolSACNFlowTest.cpp
@@ -69,7 +69,7 @@ struct FDMXProtocolLauncher
 
 IMPLEMENT_SIMPLE_AUTOMATION_TEST(DMXPrtotocolSACNBasicFlowTest, "VirtualProduction.DMX.SACN.BasicFlow", EAutomationTestFlags::ApplicationContextMask | EAutomationTestFlags::EngineFilter)
 
-static void SensDMXFragment(DMXPrtotocolSACNBasicFlowTest& Test, FDMXProtocolSACN* DMXProtocol, TSharedPtr<FDMXProtocolDeviceSACN> InProtocolDevice, uint8 InPortID)
+static void SensDMXFragment(FAutomationTestBase& Test, FDMXProtocolSACN* DMXProtocol, TSharedPtr<FDMXProtocolDeviceSACN> InProtocolDevice, uint8 InPortID)
 {
 	IDMXProtocolPort* SACNPort = DMXProtocol->GetPortManager()->GetPortByDeviceAndID(InProtocolDevice, InPortID, EDMXPortDirection::DMX_PORT_OUTPUT);
 	Test.TestNotNull(TEXT("SACNPort is valid"), SACNPort);
