UE-85265
add Copyright Notice to DMX header file

diff --git a/Engine/Plugins/Experimental/DigitalMultiplex/DMXProtocol/Source/DMXProtocol/Public/DMXProtocolCommon.h b/Engine/Plugins/Experimental/DigitalMultiplex/DMXProtocol/Source/DMXProtocol/Public/DMXProtocolCommon.h
--- a/Engine/Plugins/Experimental/DigitalMultiplex/DMXProtocol/Source/DMXProtocol/Public/DMXProtocolCommon.h
+++ b/Engine/Plugins/Experimental/DigitalMultiplex/DMXProtocol/Source/DMXProtocol/Public/DMXProtocolCommon.h
@@ -1,3 +1,5 @@
+// Copyright 1998-2019 Epic Games, Inc. All Rights Reserved.
+
 #pragma once
 
 #include "CoreMinimal.h"
