Add missing include. Add missing copyright warning.

diff --git a/Engine/Plugins/Experimental/DigitalMultiplex/DMXProtocol/Source/DMXProtocol/Public/Serialization/DMXMemoryWriter.h b/Engine/Plugins/Experimental/DigitalMultiplex/DMXProtocol/Source/DMXProtocol/Public/Serialization/DMXMemoryWriter.h
--- a/Engine/Plugins/Experimental/DigitalMultiplex/DMXProtocol/Source/DMXProtocol/Public/Serialization/DMXMemoryWriter.h
+++ b/Engine/Plugins/Experimental/DigitalMultiplex/DMXProtocol/Source/DMXProtocol/Public/Serialization/DMXMemoryWriter.h
@@ -1,3 +1,5 @@
+// Copyright 1998-2019 Epic Games, Inc. All Rights Reserved.
+
 #pragma once
 
 #include "Serialization/BufferArchive.h"
diff --git a/Engine/Plugins/Experimental/DigitalMultiplex/DMXProtocol/Source/DMXProtocolArtNet/Private/Tests/DMXPrtotocolArtNetFlowTest.cpp b/Engine/Plugins/Experimental/DigitalMultiplex/DMXProtocol/Source/DMXProtocolArtNet/Private/Tests/DMXPrtotocolArtNetFlowTest.cpp
--- a/Engine/Plugins/Experimental/DigitalMultiplex/DMXProtocol/Source/DMXProtocolArtNet/Private/Tests/DMXPrtotocolArtNetFlowTest.cpp
+++ b/Engine/Plugins/Experimental/DigitalMultiplex/DMXProtocol/Source/DMXProtocolArtNet/Private/Tests/DMXPrtotocolArtNetFlowTest.cpp
@@ -2,6 +2,7 @@
 
 #include "CoreMinimal.h"
 #include "Engine/Engine.h"
+#include "Misc/AutomationTest.h"
 #include "DMXProtocolArtNetModule.h"
 
 #include "DMXProtocolArtNet.h"
diff --git a/Engine/Plugins/Experimental/DigitalMultiplex/DMXProtocol/Source/DMXProtocolSACN/Private/Tests/DMXPrtotocolSACNFlowTest.cpp b/Engine/Plugins/Experimental/DigitalMultiplex/DMXProtocol/Source/DMXProtocolSACN/Private/Tests/DMXPrtotocolSACNFlowTest.cpp
--- a/Engine/Plugins/Experimental/DigitalMultiplex/DMXProtocol/Source/DMXProtocolSACN/Private/Tests/DMXPrtotocolSACNFlowTest.cpp
+++ b/Engine/Plugins/Experimental/DigitalMultiplex/DMXProtocol/Source/DMXProtocolSACN/Private/Tests/DMXPrtotocolSACNFlowTest.cpp
@@ -2,6 +2,7 @@
 
 #include "CoreMinimal.h"
 #include "Engine/Engine.h"
+#include "Misc/AutomationTest.h"
 #include "DMXProtocolSACNModule.h"
 
 #include "DMXProtocolSACN.h"
