Fix DMX mem copy. Fix DMX boilerplate.

diff --git a/Engine/Plugins/Experimental/DigitalMultiplex/DMXProtocol/Source/DMXProtocolArtNet/Private/DMXProtocolDeviceArtNet.cpp b/Engine/Plugins/Experimental/DigitalMultiplex/DMXProtocol/Source/DMXProtocolArtNet/Private/DMXProtocolDeviceArtNet.cpp
--- a/Engine/Plugins/Experimental/DigitalMultiplex/DMXProtocol/Source/DMXProtocolArtNet/Private/DMXProtocolDeviceArtNet.cpp
+++ b/Engine/Plugins/Experimental/DigitalMultiplex/DMXProtocol/Source/DMXProtocolArtNet/Private/DMXProtocolDeviceArtNet.cpp
@@ -1,3 +1,5 @@
+// Copyright 1998-2019 Epic Games, Inc. All Rights Reserved.
+
 #include "DMXProtocolDeviceArtNet.h"
 #include "Interfaces/IDMXProtocol.h"
 #include "Interfaces/IDMXProtocolInterface.h"
diff --git a/Engine/Plugins/Experimental/DigitalMultiplex/DMXProtocol/Source/DMXProtocolArtNet/Private/Tests/DMXPrtotocolArtNetFlowTest.cpp b/Engine/Plugins/Experimental/DigitalMultiplex/DMXProtocol/Source/DMXProtocolArtNet/Private/Tests/DMXPrtotocolArtNetFlowTest.cpp
--- a/Engine/Plugins/Experimental/DigitalMultiplex/DMXProtocol/Source/DMXProtocolArtNet/Private/Tests/DMXPrtotocolArtNetFlowTest.cpp
+++ b/Engine/Plugins/Experimental/DigitalMultiplex/DMXProtocol/Source/DMXProtocolArtNet/Private/Tests/DMXPrtotocolArtNetFlowTest.cpp
@@ -139,7 +139,7 @@ static void SensDMXFragment(FAutomationTestBase& Test, const FDMXProtocolLaunche
 
 	FRDMUID RDMUID3;
 	uint8 RDMBuffer3[RDM_UID_WIDTH] = { 0x12, 0x13, 0x14, 0x15, 0x16, 0x17 };
-	FMemory::Memcmp(RDMUID3.Buffer, RDMBuffer3, RDM_UID_WIDTH);
+	FMemory::Memcpy(RDMUID3.Buffer, RDMBuffer3, RDM_UID_WIDTH);
 	ArtNetPort->AddTODUID(RDMUID3);
 
 	// Test wrong size of RDM Buffer C array
@@ -161,7 +161,6 @@ static void SensDMXFragment(FAutomationTestBase& Test, const FDMXProtocolLaunche
 	RDMData[47] = 0x00;
 	RDMData[48] = 0xcc;
 	DMXProtocolLauncher.DMXProtocol->TransmitRDM(DMXProtocolLauncher.ProtocolDevice, DMXProtocolLauncher.PortID, RDMData);
-
 }
 
 bool FDMXPrtotocolArtNetBasicFlowTest::RunTest(const FString& Parameters)
