Add public utils for ArtNet protocol

diff --git a/Engine/Plugins/Experimental/DigitalMultiplex/DMXProtocol/Source/DMXProtocolArtNet/Private/DMXProtocolArtNetUtils.cpp b/Engine/Plugins/Experimental/DigitalMultiplex/DMXProtocol/Source/DMXProtocolArtNet/Private/DMXProtocolArtNetUtils.cpp
new file mode 100644
--- /dev/null
+++ b/Engine/Plugins/Experimental/DigitalMultiplex/DMXProtocol/Source/DMXProtocolArtNet/Private/DMXProtocolArtNetUtils.cpp
@@ -0,0 +1,33 @@
+// Copyright 1998-2019 Epic Games, Inc. All Rights Reserved.
+
+#include "DMXProtocolArtNetUtils.h"
+#include "DMXProtocolArtNetPublicUtils.h"
+
+namespace ArtNet
+{
+	uint16 GetPacketType(const FArrayReaderPtr& Buffer)
+	{
+		uint16 OpCode = 0x0000;
+		const uint32 MinCheck = ARTNET_STRING_SIZE + 2;
+		if (Buffer->Num() > MinCheck)
+		{
+			// Get OpCode
+			Buffer->Seek(ARTNET_STRING_SIZE);
+			*Buffer << OpCode;
+
+			// Reset Position
+			Buffer->Seek(0);
+		}
+
+		return OpCode;
+	}
+
+	uint16 ComputeUniverseID(const uint8 Net, const uint8 Subnet, const uint8 Universe)
+	{
+		/* Bits 15 			Bits 14-8        | Bits 7-4      | Bits 3-0
+		* 0 				Net 			 | Sub-Net 		 | Universe
+		* 0				    (0b1111111 << 8) | (0b1111 << 4) | (0b1111)
+		*/
+		return (Net << 8) | (Subnet << 4) | Universe;
+	}
+}
diff --git a/Engine/Plugins/Experimental/DigitalMultiplex/DMXProtocol/Source/DMXProtocolArtNet/Private/DMXProtocolArtNetUtils.h b/Engine/Plugins/Experimental/DigitalMultiplex/DMXProtocol/Source/DMXProtocolArtNet/Private/DMXProtocolArtNetUtils.h
--- a/Engine/Plugins/Experimental/DigitalMultiplex/DMXProtocol/Source/DMXProtocolArtNet/Private/DMXProtocolArtNetUtils.h
+++ b/Engine/Plugins/Experimental/DigitalMultiplex/DMXProtocol/Source/DMXProtocolArtNet/Private/DMXProtocolArtNetUtils.h
@@ -10,20 +10,6 @@
 
 namespace ArtNet
 {
-	static uint16 GetPacketType(const FArrayReaderPtr& Buffer)
-	{
-		uint16 OpCode = 0x0000;
-		const uint32 MinCheck = ARTNET_STRING_SIZE + 2;
-		if (Buffer->Num() > MinCheck)
-		{
-			// Get OpCode
-			Buffer->Seek(ARTNET_STRING_SIZE);
-			*Buffer << OpCode;
-
-			// Reset Position
-			Buffer->Seek(0);
-		}
-
-		return OpCode;
-	}
+	/** Get package 2 bytes ID from the package */
+	uint16 GetPacketType(const FArrayReaderPtr& Buffer);
 }
\ No newline at end of file
diff --git a/Engine/Plugins/Experimental/DigitalMultiplex/DMXProtocol/Source/DMXProtocolArtNet/Public/DMXProtocolArtNetPublicUtils.h b/Engine/Plugins/Experimental/DigitalMultiplex/DMXProtocol/Source/DMXProtocolArtNet/Public/DMXProtocolArtNetPublicUtils.h
new file mode 100644
--- /dev/null
+++ b/Engine/Plugins/Experimental/DigitalMultiplex/DMXProtocol/Source/DMXProtocolArtNet/Public/DMXProtocolArtNetPublicUtils.h
@@ -0,0 +1,11 @@
+// Copyright 1998-2019 Epic Games, Inc. All Rights Reserved.
+
+#pragma once
+
+#include "CoreMinimal.h"
+
+namespace ArtNet
+{
+	/** Generate and return a UniverseID value from Net, Sub-Net and Universe values. */
+	DMXPROTOCOLARTNET_API uint16 ComputeUniverseID(const uint8 Net, const uint8 Subnet, const uint8 Universe);
+}
