- DMX Editor

diff --git a/Engine/Plugins/Experimental/DigitalMultiplex/DMXEngine/Source/DMXRuntime/Private/Library/DMXEntityController.cpp b/Engine/Plugins/Experimental/DigitalMultiplex/DMXEngine/Source/DMXRuntime/Private/Library/DMXEntityController.cpp
new file mode 100644
--- /dev/null
+++ b/Engine/Plugins/Experimental/DigitalMultiplex/DMXEngine/Source/DMXRuntime/Private/Library/DMXEntityController.cpp
@@ -0,0 +1,81 @@
+// Copyright Epic Games, Inc. All Rights Reserved.
+
+#include "Library/DMXEntityController.h"
+#include "Interfaces/IDMXProtocol.h"
+
+#if WITH_EDITOR
+
+void UDMXEntityController::PostEditChangeProperty(FPropertyChangedEvent& PropertyChangedEvent)
+{
+	const FName&& PropertyName = PropertyChangedEvent.GetPropertyName();
+	if (PropertyName == GET_MEMBER_NAME_CHECKED(UDMXEntityController, UniverseLocalStart)
+		|| PropertyName == GET_MEMBER_NAME_CHECKED(UDMXEntityController, UniverseLocalNum)
+		|| PropertyName == GET_MEMBER_NAME_CHECKED(UDMXEntityController, UniverseRemoteStart)
+		|| PropertyName == GET_MEMBER_NAME_CHECKED(UDMXEntityController, DeviceProtocol))
+	{
+		ValidateRangeValues();
+		UpdateUniversesFromRange();
+	}
+
+	Super::PostEditChangeProperty(PropertyChangedEvent);
+}
+
+#endif // WITH_EDITOR
+
+void UDMXEntityController::ValidateRangeValues()
+{
+	if (DeviceProtocol.IsValid())
+	{
+		const TSharedPtr<IDMXProtocol>& Protocol = DeviceProtocol;
+		const int32 MinUniverseID = Protocol->GetMinUniverseID();
+		const int32 MaxUniverseID = Protocol->GetMaxUniverses();
+
+		// To make sure all protocols have a minimum local value of 1, we offset their min values
+		const int32 LocalMinOffset = 1 - MinUniverseID;
+
+		// Clamp local values
+		UniverseLocalStart = FMath::Clamp(UniverseLocalStart, 1, MaxUniverseID + LocalMinOffset);
+		UniverseLocalNum = FMath::Clamp(UniverseLocalNum, 1, MaxUniverseID + LocalMinOffset - UniverseLocalStart + 1);
+
+		UniverseLocalEnd = UniverseLocalStart + UniverseLocalNum - 1;
+
+		// Clamp remote start to have valid remote range
+		UniverseRemoteStart = FMath::Clamp(UniverseRemoteStart, MinUniverseID, MaxUniverseID - UniverseLocalNum + 1);
+	}
+	else
+	{
+		if (UniverseLocalStart < 0)
+		{
+			UniverseLocalStart = 0; 
+		}
+		if (UniverseLocalNum < 1)
+		{
+			UniverseLocalNum = 1; 
+		}
+
+		UniverseLocalEnd = UniverseLocalStart + UniverseLocalNum - 1;
+	}
+
+	UniverseRemoteEnd = UniverseRemoteStart + UniverseLocalNum - 1;
+	RemoteOffset = UniverseRemoteStart - UniverseLocalStart;
+}
+
+void UDMXEntityController::UpdateUniversesFromRange()
+{
+	const uint16 NumUniverses(UniverseLocalNum);
+
+	if (NumUniverses < Universes.Num())
+	{
+		Universes.Reset(NumUniverses);
+		Universes.AddZeroed(NumUniverses);
+	}
+	else if (NumUniverses > Universes.Num())
+	{
+		Universes.AddZeroed(NumUniverses - Universes.Num());
+	}
+
+	for (uint16 UniverseIndex = 0, UniverseID = UniverseRemoteStart; UniverseID <= UniverseRemoteEnd; ++UniverseIndex, ++UniverseID)
+	{
+		Universes[UniverseIndex].UniverseNumber = UniverseID;
+	}
+}
diff --git a/Engine/Plugins/Experimental/DigitalMultiplex/DMXEngine/Source/DMXRuntime/Private/Library/DMXEntityFixturePatch.cpp b/Engine/Plugins/Experimental/DigitalMultiplex/DMXEngine/Source/DMXRuntime/Private/Library/DMXEntityFixturePatch.cpp
new file mode 100644
--- /dev/null
+++ b/Engine/Plugins/Experimental/DigitalMultiplex/DMXEngine/Source/DMXRuntime/Private/Library/DMXEntityFixturePatch.cpp
@@ -0,0 +1,339 @@
+// Copyright Epic Games, Inc. All Rights Reserved.
+
+#include "Library/DMXEntityFixturePatch.h"
+
+#include "Library/DMXLibrary.h"
+#include "Library/DMXEntityController.h"
+#include "Library/DMXEntityFixtureType.h"
+#include "DMXProtocolConstants.h"
+
+#define LOCTEXT_NAMESPACE "DMXEntityFixturePatch"
+
+DECLARE_LOG_CATEGORY_CLASS(DMXEntityFixtureTypePatchLog, Log, All);
+
+int32 UDMXEntityFixturePatch::GetChannelSpan() const
+{
+	if (ParentFixtureTypeTemplate != nullptr)
+	{
+		// Check for existing modes before trying to read them
+		const int32 NumModes = ParentFixtureTypeTemplate->Modes.Num();
+		if (NumModes > 0 && ActiveMode < NumModes)
+		{
+			// Number of channels occupied by all of the Active Mode's functions
+			return ParentFixtureTypeTemplate->Modes[ActiveMode].ChannelSpan;
+		}
+	}
+
+	return 1;
+}
+
+int32 UDMXEntityFixturePatch::GetStartingChannel() const
+{
+	if (bAutoAssignAddress)
+	{
+		return AutoStartingAddress;
+	}
+	else
+	{
+		return ManualStartingAddress;
+	}
+}
+
+TArray<FName> UDMXEntityFixturePatch::GetAllFunctionsInActiveMode() const
+{
+	TArray<FName> NameArray;
+
+	if (!CanReadActiveMode())
+	{
+		return NameArray;
+	}
+
+	const FDMXFixtureMode& Mode = ParentFixtureTypeTemplate->Modes[ActiveMode];
+	const TArray<FDMXFixtureFunction>& Functions = Mode.Functions;
+	NameArray.Reserve(Functions.Num());
+	for (const FDMXFixtureFunction& Function : Functions)
+	{
+		if (UDMXEntityFixtureType::GetFunctionLastChannel(Function) <= Mode.ChannelSpan)
+		{
+			NameArray.Add(FName(*Function.FunctionName));
+		}
+	}
+	return NameArray;
+}
+
+TMap<FName, int32> UDMXEntityFixturePatch::GetFunctionDefaultMap() const
+{
+	TMap<FName, int32> DefaultValueMap;
+
+	if (!CanReadActiveMode())
+	{
+		return DefaultValueMap;
+	}
+
+	const FDMXFixtureMode& Mode = ParentFixtureTypeTemplate->Modes[ActiveMode];
+	const TArray<FDMXFixtureFunction>& Functions = Mode.Functions;
+	DefaultValueMap.Reserve(Functions.Num());
+	for (const FDMXFixtureFunction& Function : Functions)
+	{
+		if (UDMXEntityFixtureType::GetFunctionLastChannel(Function) <= Mode.ChannelSpan)
+		{
+			DefaultValueMap.Add(FName(*Function.FunctionName), Function.DefaultValue);
+		}
+	}
+	return DefaultValueMap;
+}
+
+TMap<FName, int32> UDMXEntityFixturePatch::GetFunctionChannelAssignments() const
+{
+	TMap<FName, int32> ChannelMap;
+
+	if (!CanReadActiveMode())
+	{
+		return ChannelMap;
+	}
+
+	const FDMXFixtureMode& Mode = ParentFixtureTypeTemplate->Modes[ActiveMode];
+	const TArray<FDMXFixtureFunction>& Functions = Mode.Functions;
+	ChannelMap.Reserve(Functions.Num());
+	for (const FDMXFixtureFunction& Function : Functions)
+	{
+		if (UDMXEntityFixtureType::GetFunctionLastChannel(Function) <= Mode.ChannelSpan)
+		{
+			ChannelMap.Add(FName(*Function.FunctionName), Function.Channel);
+		}
+		ChannelMap.Add(FName(*Function.FunctionName), Function.Channel + GetStartingChannel() - 1);
+	}
+	return ChannelMap;
+}
+
+TMap<FName, EDMXFixtureSignalFormat> UDMXEntityFixturePatch::GetFunctionSignalFormats() const
+{
+	TMap<FName, EDMXFixtureSignalFormat> FormatMap;
+
+	if (!CanReadActiveMode())
+	{
+		return FormatMap;
+	}
+
+	const FDMXFixtureMode& Mode = ParentFixtureTypeTemplate->Modes[ActiveMode];
+	const TArray<FDMXFixtureFunction>& Functions = Mode.Functions;
+	FormatMap.Reserve(Functions.Num());
+	for (const FDMXFixtureFunction& Function : Functions)
+	{
+		if (UDMXEntityFixtureType::GetFunctionLastChannel(Function) <= Mode.ChannelSpan)
+		{
+			FormatMap.Add(FName(*Function.FunctionName), Function.DataType);
+		}
+	}
+	return FormatMap;
+}
+
+TMap<FName, int32> UDMXEntityFixturePatch::ConvertRawMapToFunctionMap(const TMap<int32, uint8>& RawMap) const
+{
+	TMap<FName, int32> FunctionMap;
+
+	if (!CanReadActiveMode())
+	{
+		return FunctionMap;
+	}
+
+	const FDMXFixtureMode& Mode = ParentFixtureTypeTemplate->Modes[ActiveMode];
+	const TArray<FDMXFixtureFunction>& Functions = Mode.Functions;
+	// Let's consider all functions in the raw map are 8bit and allocate for 1 channel = 1 function
+	// We'll avoid many allocations but potentially have a map 4x the size we need. We can shrink it later
+	FunctionMap.Reserve(RawMap.Num());
+
+	for (const FDMXFixtureFunction& Function : Functions)
+	{
+		// Ignore functions outside the Active Mode's Channel Span
+		if (UDMXEntityFixtureType::GetFunctionLastChannel(Function) <= Mode.ChannelSpan
+			&& RawMap.Contains(Function.Channel))
+		{
+			const uint8& RawValue(RawMap.FindRef(Function.Channel));
+			const uint8 ChannelsToAdd = UDMXEntityFixtureType::NumChannelsToOccupy(Function.DataType);
+
+			uint32 IntVal = 0;
+			for (uint8 ChannelIt = 0; ChannelIt < ChannelsToAdd; ChannelIt++)
+			{
+				if (const uint8* RawVal = RawMap.Find(Function.Channel + ChannelIt))
+				{
+					IntVal += *RawVal << (ChannelIt * 8);
+				}
+			}
+
+			FunctionMap.Add(*Function.FunctionName, IntVal);
+		}
+	}
+
+	return FunctionMap;
+}
+
+TMap<int32, uint8> UDMXEntityFixturePatch::ConvertFunctionMapToRawMap(const TMap<FName, int32>& FunctionMap) const
+{
+	TMap<int32, uint8> RawMap;
+
+	if (!CanReadActiveMode())
+	{
+		return RawMap;
+	}
+
+	const FDMXFixtureMode& Mode = ParentFixtureTypeTemplate->Modes[ActiveMode];
+	const TArray<FDMXFixtureFunction>& Functions = Mode.Functions;
+
+	RawMap.Reserve(FunctionMap.Num() * 4); // Let's assume all functions are 32bit. We can shrink RawMap later.
+
+	for (const TPair<FName, int32>& Elem : FunctionMap)
+	{
+		// Search for a function with Elem.Key as name
+		const FDMXFixtureFunction* FunctionPtr = Functions.FindByPredicate([&Elem] (const FDMXFixtureFunction& Function)
+			{
+				return Elem.Key == FName(*Function.FunctionName);
+			});
+
+		// Also check for the Function being in the valid range for the Active Mode's channel span
+		if (FunctionPtr != nullptr && Mode.ChannelSpan >= UDMXEntityFixtureType::GetFunctionLastChannel(*FunctionPtr))
+		{
+			const uint8&& NumChannels = UDMXEntityFixtureType::NumChannelsToOccupy(FunctionPtr->DataType);
+			const uint32 Value = UDMXEntityFixtureType::ClampValueToDataType(FunctionPtr->DataType, Elem.Value);
+
+			for (uint8 ChannelIt = 0; ChannelIt < NumChannels; ChannelIt++)
+			{
+				const uint8 BytesOffset = (ChannelIt) * 8;
+				const uint8 ChannelVal = Value >> BytesOffset & 0xff;
+
+				const int32 FinalChannel = FunctionPtr->Channel + (GetStartingChannel() - 1);
+				RawMap.Add(FinalChannel + ChannelIt, ChannelVal);
+			}
+		}
+	}
+
+	RawMap.Shrink();
+	return RawMap;
+}
+
+bool UDMXEntityFixturePatch::IsMapValid(const TMap<FName, int32>& FunctionMap) const
+{
+	if (!CanReadActiveMode())
+	{
+		if (FunctionMap.Num() == 0)
+		{
+			return true;
+		}
+		else
+		{
+			return false;
+		}
+	}
+
+	const TArray<FDMXFixtureFunction>& Functions = ParentFixtureTypeTemplate->Modes[ActiveMode].Functions;
+
+	for (const TPair<FName, int32>& Elem : FunctionMap)
+	{
+		if (!ContainsFunction(Elem.Key))
+		{
+			return false;
+		}
+	}
+	return true;
+}
+
+TMap<FName, int32> UDMXEntityFixturePatch::ConvertToValidMap(const TMap<FName, int32>& FunctionMap) const
+{
+	TMap<FName, int32> ValidMap;
+
+	if (!CanReadActiveMode())
+	{
+		return ValidMap;
+	}
+
+	const FDMXFixtureMode& Mode = ParentFixtureTypeTemplate->Modes[ActiveMode];
+	const TArray<FDMXFixtureFunction>& Functions = Mode.Functions;
+	ValidMap.Reserve(FunctionMap.Num());
+
+	for (const TPair<FName, int32>& Elem : FunctionMap)
+	{
+		// Search for a function with Elem.Key as name
+		const FDMXFixtureFunction* FunctionPtr = Functions.FindByPredicate([&Elem](const FDMXFixtureFunction& Function)
+			{
+				return Elem.Key == FName(*Function.FunctionName);
+			});
+
+		// Also check for the Function being in the valid range for the Active Mode's channel span
+		if (FunctionPtr != nullptr && Mode.ChannelSpan >= UDMXEntityFixtureType::GetFunctionLastChannel(*FunctionPtr))
+		{
+			ValidMap.Add(Elem.Key, Elem.Value);
+		}
+	}
+
+	ValidMap.Shrink();
+	return ValidMap;
+}
+
+bool UDMXEntityFixturePatch::IsValidEntity(FText& OutReason) const
+{
+	OutReason = FText::GetEmpty();
+
+	if (ParentFixtureTypeTemplate == nullptr)
+	{
+		OutReason = LOCTEXT("InvalidReason_NullParentTemplate", "Fixture Template is null");
+	}
+	else if ((GetStartingChannel() + GetChannelSpan() - 1) > DMX_UNIVERSE_SIZE)
+	{
+		OutReason = FText::Format(
+			LOCTEXT("InvalidReason_ChannelOverflow", "Channels range overflows max channel address ({0})"),
+			FText::AsNumber(DMX_UNIVERSE_SIZE)
+		);
+	}
+
+	return OutReason.IsEmpty();
+}
+
+void UDMXEntityFixturePatch::ValidateActiveMode()
+{
+	if (ParentFixtureTypeTemplate != nullptr)
+	{
+		ActiveMode = FMath::Clamp(ActiveMode, 0, ParentFixtureTypeTemplate->Modes.Num() - 1);
+	}
+}
+
+TArray<UDMXEntityController*> UDMXEntityFixturePatch::GetRelevantControllers() const
+{
+	TArray<UDMXEntityController*> RetVal;
+	if (ParentLibrary != nullptr)
+	{
+		ParentLibrary->ForEachEntityOfType<UDMXEntityController>([&](UDMXEntityController* Controller)
+			{
+				if (IsInControllerRange(Controller))
+				{
+					RetVal.Add(Controller);
+				}
+			});
+	}
+	else
+	{
+		UE_LOG(DMXEntityFixtureTypePatchLog, Fatal, TEXT("Parent library is null!"));
+	}
+
+	return RetVal;
+}
+
+bool UDMXEntityFixturePatch::IsInControllersRange(const TArray<UDMXEntityController*>& InControllers) const
+{
+	for (const UDMXEntityController* Controller : InControllers)
+	{
+		if (IsInControllerRange(Controller))
+		{
+			return true;
+		}
+	}
+	return false;
+}
+
+UDMXEntityFixturePatch::UDMXEntityFixturePatch()
+	: bAutoAssignAddress(true)
+	, ManualStartingAddress(1)
+	, AutoStartingAddress(1)
+	, ActiveMode(0)
+{}
+
+#undef LOCTEXT_NAMESPACE
diff --git a/Engine/Plugins/Experimental/DigitalMultiplex/DMXEngine/Source/DMXRuntime/Private/Library/DMXEntityFixtureType.cpp b/Engine/Plugins/Experimental/DigitalMultiplex/DMXEngine/Source/DMXRuntime/Private/Library/DMXEntityFixtureType.cpp
new file mode 100644
--- /dev/null
+++ b/Engine/Plugins/Experimental/DigitalMultiplex/DMXEngine/Source/DMXRuntime/Private/Library/DMXEntityFixtureType.cpp
@@ -0,0 +1,267 @@
+// Copyright Epic Games, Inc. All Rights Reserved.
+
+#include "Library/DMXEntityFixtureType.h"
+#include "Library/DMXEntityFixturePatch.h"
+#include "Library/DMXLibrary.h"
+#include "Library/DMXImport.h"
+#include "Library/DMXImportGDTF.h"
+
+void UDMXEntityFixtureType::SetModesFromDMXImport(UDMXImport* DMXImportAsset)
+{
+	if (DMXImportAsset == nullptr || !DMXImportAsset->IsValidLowLevelFast())
+	{
+		return; 
+	}
+
+	if (UDMXImportGDTFDMXModes* GDTFDMXModes = Cast<UDMXImportGDTFDMXModes>(DMXImportAsset->DMXModes))
+	{
+		// Clear existing modes
+		Modes.Empty(DMXImportAsset->DMXModes != nullptr ? GDTFDMXModes->DMXModes.Num() : 0);
+
+		if (DMXImportAsset->DMXModes == nullptr) 
+		{ 
+			return; 
+		}
+
+		// Copy modes from asset
+		for (const FDMXImportGDTFDMXMode& AssetMode : GDTFDMXModes->DMXModes)
+		{
+			FDMXFixtureMode& Mode = Modes[Modes.Emplace()];
+			Mode.ModeName = AssetMode.Name.ToString();
+
+			for (const FDMXImportGDTFDMXChannel& ModeChannel : AssetMode.DMXChannels)
+			{
+				FDMXFixtureFunction& Function = Mode.Functions[Mode.Functions.Emplace()];
+				Function.FunctionName = ModeChannel.LogicalChannel.ChannelFunction.Name.ToString();
+				SetFunctionSize(Mode, Function, ModeChannel.Offset.Num());
+				Function.DefaultValue = ModeChannel.Default.Value;
+			}
+		}
+	}
+}
+
+void UDMXEntityFixtureType::SetFunctionSize(FDMXFixtureMode& InMode, FDMXFixtureFunction& InFunction, uint8 Size)
+{
+	//Get New Data Type
+	EDMXFixtureSignalFormat NewDataType;
+	switch (Size)
+	{
+	case 0:
+	case 1:
+		NewDataType = EDMXFixtureSignalFormat::E8Bit;
+		break;
+	case 2:
+		NewDataType = EDMXFixtureSignalFormat::E16Bit;
+		break;
+	case 3:
+		NewDataType = EDMXFixtureSignalFormat::E24Bit;
+		break;
+	case 4:
+	default:
+		NewDataType = EDMXFixtureSignalFormat::E32Bit;
+		break;
+	}
+
+#if WITH_EDITOR
+	//Update UI
+	InFunction.DataType = NewDataType;
+	UpdateModeChannelProperties(InMode);
+	ClampDefaultValue(InFunction);
+#endif // WITH_EDITOR
+}
+
+uint8 UDMXEntityFixtureType::GetFunctionLastChannel(const FDMXFixtureFunction& Function)
+{
+	return Function.Channel + UDMXEntityFixtureType::NumChannelsToOccupy(Function.DataType) - 1;
+}
+
+void UDMXEntityFixtureType::ClampDefaultValue(FDMXFixtureFunction& InFunction)
+{
+	int64& DefaultValue = InFunction.DefaultValue;
+	switch (InFunction.DataType)
+	{
+	case EDMXFixtureSignalFormat::E8BitSubFunctions:
+	case EDMXFixtureSignalFormat::E8Bit:
+		DefaultValue = FMath::Clamp(DefaultValue, (int64)0, (int64)MAX_uint8);
+		break;
+	case EDMXFixtureSignalFormat::E16Bit:
+		DefaultValue = FMath::Clamp(DefaultValue, (int64)0, (int64)MAX_uint16);
+		break;
+	case EDMXFixtureSignalFormat::E24Bit:
+		DefaultValue = FMath::Clamp(DefaultValue, (int64)0, (int64)0xFFFFFF);
+		break;
+	case EDMXFixtureSignalFormat::E32Bit:
+		DefaultValue = FMath::Clamp(DefaultValue, (int64)0, (int64)MAX_uint32);
+		break;
+	default:
+		checkNoEntry();
+		break;
+	}
+}
+
+uint8 UDMXEntityFixtureType::NumChannelsToOccupy(EDMXFixtureSignalFormat DataType)
+{
+	switch (DataType)
+	{
+	case EDMXFixtureSignalFormat::E8BitSubFunctions:
+	case EDMXFixtureSignalFormat::E8Bit:
+		return 1;
+
+	case EDMXFixtureSignalFormat::E16Bit:
+		return 2;
+
+	case EDMXFixtureSignalFormat::E24Bit:
+		return 3;
+
+	case EDMXFixtureSignalFormat::E32Bit:
+		return 4;
+
+	default:
+		break;
+	}
+	return 1;
+}
+
+uint32 UDMXEntityFixtureType::ClampValueToDataType(EDMXFixtureSignalFormat DataType, uint32 InValue)
+{
+	switch (DataType)
+	{
+	case EDMXFixtureSignalFormat::E8BitSubFunctions:
+	case EDMXFixtureSignalFormat::E8Bit:
+		return FMath::Clamp(InValue, 0u, (uint32)MAX_uint8);
+
+	case EDMXFixtureSignalFormat::E16Bit:
+		return FMath::Clamp(InValue, 0u, (uint32)MAX_uint16);
+
+	case EDMXFixtureSignalFormat::E24Bit:
+		return FMath::Clamp(InValue, 0u, 0xFFFFFFu);
+
+	case EDMXFixtureSignalFormat::E32Bit:
+		return FMath::Clamp(InValue, 0u, MAX_uint32);
+
+	default:
+		break;
+	}
+	return InValue;
+}
+
+#if WITH_EDITOR
+
+void UDMXEntityFixtureType::PostEditChangeChainProperty(FPropertyChangedChainEvent& PropertyChangedEvent)
+{
+	// Clamp DefaultValue from selected data type
+	if (PropertyChangedEvent.GetPropertyName() == GET_MEMBER_NAME_CHECKED(FDMXFixtureFunction, DefaultValue)
+		|| PropertyChangedEvent.GetPropertyName() == GET_MEMBER_NAME_CHECKED(FDMXFixtureFunction, DataType))
+	{
+		const int32 ModeIndex = PropertyChangedEvent.GetArrayIndex(GET_MEMBER_NAME_CHECKED(UDMXEntityFixtureType, Modes).ToString());
+		const int32 FunctionIndex = PropertyChangedEvent.GetArrayIndex(GET_MEMBER_NAME_CHECKED(FDMXFixtureMode, Functions).ToString());
+
+		check(ModeIndex != INDEX_NONE && FunctionIndex != INDEX_NONE);
+		FDMXFixtureFunction& Function = Modes[ModeIndex].Functions[FunctionIndex];
+
+		ClampDefaultValue(Function);
+	}
+
+	// Refresh ChannelSpan from functions' settings
+	if (PropertyChangedEvent.GetPropertyName() == GET_MEMBER_NAME_CHECKED(FDMXFixtureMode, bAutoChannelSpan)
+		|| PropertyChangedEvent.GetPropertyName() == GET_MEMBER_NAME_CHECKED(FDMXFixtureMode, Functions)
+		|| PropertyChangedEvent.GetPropertyName() == GET_MEMBER_NAME_CHECKED(FDMXFixtureFunction, ChannelOffset)
+		|| PropertyChangedEvent.GetPropertyName() == GET_MEMBER_NAME_CHECKED(FDMXFixtureFunction, DataType))
+	{
+		// If we have a specific Modes index that was modified, update its properties
+		const int32 ModeIndex = PropertyChangedEvent.GetArrayIndex(GET_MEMBER_NAME_CHECKED(UDMXEntityFixtureType, Modes).ToString());
+		if (ModeIndex != INDEX_NONE)
+		{
+			UpdateModeChannelProperties(Modes[ModeIndex]);
+		}
+		else
+		{
+			// Unfortunately, some operations like reordering an array's values, don't give us an array index
+
+			// Check if this is a property contained in a Modes property
+			const auto HeadNode = PropertyChangedEvent.PropertyChain.GetHead();
+			if (HeadNode != nullptr && HeadNode->GetValue() != nullptr
+				&& HeadNode->GetValue()->GetFName() == GET_MEMBER_NAME_CHECKED(UDMXEntityFixtureType, Modes))
+			{
+				for (FDMXFixtureMode& Mode : Modes)
+				{
+					UpdateModeChannelProperties(Mode);
+				}
+			}
+		}
+	}
+
+	if (PropertyChangedEvent.GetPropertyName() == GET_MEMBER_NAME_CHECKED(UDMXEntityFixtureType, Modes) &&
+        (PropertyChangedEvent.ChangeType == EPropertyChangeType::ArrayRemove || PropertyChangedEvent.ChangeType == EPropertyChangeType::ArrayClear))
+	{
+		// Warn patches from this type about the Mode(s) removal so they can keep their ActiveMode value valid
+		if (ParentLibrary != nullptr)
+		{
+			ParentLibrary->ForEachEntityOfType<UDMXEntityFixturePatch>([this](UDMXEntityFixturePatch* Patch)
+				{
+					if (Patch->ParentFixtureTypeTemplate == this)
+					{
+						Patch->ValidateActiveMode();
+					}
+				});
+		}
+	}
+
+	Super::PostEditChangeChainProperty(PropertyChangedEvent);
+}
+
+void UDMXEntityFixtureType::PostEditUndo()
+{
+	for (FDMXFixtureMode& Mode : Modes)
+	{
+		UpdateModeChannelProperties(Mode);
+	}
+
+	Super::PostEditUndo();
+}
+
+void UDMXEntityFixtureType::UpdateModeChannelProperties(FDMXFixtureMode& Mode)
+{
+	int32 ChannelSpan = 0;
+
+	for (FDMXFixtureFunction& Function : Mode.Functions)
+	{
+		Function.Channel = ChannelSpan + 1 + Function.ChannelOffset;
+
+		switch (Function.DataType)
+		{
+		case EDMXFixtureSignalFormat::E8Bit:
+		case EDMXFixtureSignalFormat::E8BitSubFunctions:
+			ChannelSpan = Function.Channel;
+			break;
+		case EDMXFixtureSignalFormat::E16Bit:
+			ChannelSpan = Function.Channel + 1;
+			break;
+		case EDMXFixtureSignalFormat::E24Bit:
+			ChannelSpan = Function.Channel + 2;
+			break;
+		case EDMXFixtureSignalFormat::E32Bit:
+			ChannelSpan = Function.Channel + 3;
+			break;
+		default:
+			checkNoEntry();
+			break;
+		}
+	}
+
+	if (ChannelSpan < 1)
+	{
+		ChannelSpan = 1; 
+	}
+	if (ChannelSpan > 512)
+	{
+		ChannelSpan = 512; 
+	}
+
+	if (Mode.bAutoChannelSpan)
+	{
+		Mode.ChannelSpan = ChannelSpan;
+	}
+}
+
+#endif // WITH_EDITOR
\ No newline at end of file
diff --git a/Engine/Plugins/Experimental/DigitalMultiplex/DMXEngine/Source/DMXRuntime/Public/DMXSubsystem.h b/Engine/Plugins/Experimental/DigitalMultiplex/DMXEngine/Source/DMXRuntime/Public/DMXSubsystem.h
new file mode 100644
--- /dev/null
+++ b/Engine/Plugins/Experimental/DigitalMultiplex/DMXEngine/Source/DMXRuntime/Public/DMXSubsystem.h
@@ -0,0 +1,139 @@
+// Copyright Epic Games, Inc. All Rights Reserved.
+
+#pragma once
+
+#include "CoreMinimal.h"
+#include "DMXProtocolTypes.h"
+#include "DMXTypes.h"
+#include "Library/DMXEntityReference.h"
+#include "Library/DMXEntityFixtureType.h"
+#include "Subsystems/EngineSubsystem.h"
+#include "DMXSubsystem.generated.h"
+
+class UDMXLibrary;
+class IDMXProtocol;
+class UDMXEntityFixturePatch;
+
+DECLARE_DYNAMIC_MULTICAST_DELEGATE_ThreeParams(FProtocolReceivedDelegate, FDMXProtocolName, Protocol, int32, Universe, const TArray<uint8>&, DMXBuffer);
+
+/**
+ * UDMXSubsystem
+ * Collections of DMX context blueprint subsystem functions and internal functions for DMX K2Nodes
+ */
+UCLASS()
+class DMXRUNTIME_API UDMXSubsystem 
+	: public UEngineSubsystem
+{
+	GENERATED_BODY()
+	
+public:
+	/**  Send DMX using function names and integer values. */
+	UFUNCTION(BlueprintCallable, Category = "DMX")
+	void SendDMX(FDMXProtocolName SelectedProtocol, UDMXEntityFixturePatch* FixturePatch, TMap<FName, int32> FunctionMap, EDMXSendResult& OutResult);
+
+	/**  Send DMX using channel and value raw values. */
+	UFUNCTION(BlueprintCallable, Category = "DMX")
+	void SendDMXRaw(FDMXProtocolName SelectedProtocol, int32 UniverseIndex, TMap<int32, uint8> ChannelValuesMap, EDMXSendResult& OutResult);
+
+	/**  Return reference to array of Fixture Patch objects of a given type. */
+	UFUNCTION(BlueprintCallable, Category = "DMX", meta = (BlueprintInternalUseOnly = "true", AutoCreateRefTerm = "FixtureType"))
+	void GetAllFixturesOfType(const UDMXLibrary* DMXLibrary, const FName& FixtureType, TArray<UDMXEntityFixturePatch*>& OutResult);
+
+	/**  Return reference to array of Fixture Patch objects of a given category. */
+	UFUNCTION(BlueprintCallable, Category = "DMX")
+	void GetAllFixturesOfCategory(const UDMXLibrary* DMXLibrary, FDMXFixtureCategory Category, TArray<UDMXEntityFixturePatch*>& OutResult);
+
+	/**  Return reference to array of Fixture Patch objects in a given universe. */
+	UFUNCTION(BlueprintCallable, Category = "DMX")
+	void GetAllFixturesInUniverse(const UDMXLibrary* DMXLibrary, int32 UniverseId, TArray<UDMXEntityFixturePatch*>& OutResult);
+
+	/**  Return reference to array of Fixture Patch objects in a given controller. */
+	UFUNCTION(BlueprintCallable, Category = "DMX")
+	void GetAllUniversesInController(const UDMXLibrary* DMXLibrary, FString ControllerName, TArray<int32>& OutResult);
+
+	/**  Return byte array from the DMX buffer given a universe. */
+	UFUNCTION(BlueprintCallable, Category = "DMX")
+	void GetRawBuffer(FDMXProtocolName SelectedProtocol, int32 UniverseIndex, TArray<uint8>& DMXBuffer);
+
+	/**  Return map with all DMX functions and their associated values given DMX buffer and desired universe. */
+	UFUNCTION(BlueprintCallable, Category = "DMX")
+	void GetFixtureFunctions(const UDMXEntityFixturePatch* InFixturePatch, const TArray<uint8>& DMXBuffer, TMap<FName, int32>& OutResult);
+
+	/**  Return reference to array of Fixture Patch objects with a given tag. */
+	UFUNCTION(BlueprintCallable, Category = "DMX")
+	TArray<UDMXEntityFixturePatch*> GetAllFixturesWithTag(const UDMXLibrary* DMXLibrary, FName CustomTag);
+
+	/**  Return reference to array of Fixture Patch objects in library. */
+	UFUNCTION(BlueprintCallable, Category = "DMX")
+	TArray<UDMXEntityFixturePatch*> GetAllFixturesInLibrary(const UDMXLibrary* DMXLibrary);
+
+	/**  Return reference to Fixture Patch object with a given name. */
+	UFUNCTION(BlueprintCallable, Category = "DMX", meta = (AutoCreateRefTerm="Name"))
+	UDMXEntityFixturePatch* GetFixtureByName(const UDMXLibrary* DMXLibrary, const FString& Name);
+
+	/**  Return reference to array of Fixture Types objects in library. */
+	UFUNCTION(BlueprintCallable, Category = "DMX")
+	TArray<UDMXEntityFixtureType*> GetAllFixtureTypesInLibrary(const UDMXLibrary* DMXLibrary);
+
+	/**  Return reference to Fixture Type object with a given name. */
+	UFUNCTION(BlueprintCallable, Category = "DMX", meta = (AutoCreateRefTerm = "Name"))
+	UDMXEntityFixtureType* GetFixtureTypeByName(const UDMXLibrary* DMXLibrary, const FString& Name);
+
+	/**  Return reference to array of Controller objects in library. */
+	UFUNCTION(BlueprintCallable, Category = "DMX")
+	TArray<UDMXEntityController*> GetAllControllersInLibrary(const UDMXLibrary* DMXLibrary);
+
+	/**  Return reference to Controller object with a given name. */
+	UFUNCTION(BlueprintCallable, Category = "DMX", meta = (AutoCreateRefTerm = "Name"))
+	UDMXEntityController* GetControllerByName(const UDMXLibrary* DMXLibrary, const FString& Name);
+
+	/**  Return reference to array of DMX Library objects. */
+	UFUNCTION(BlueprintCallable, Category = "DMX")
+	TArray<UDMXLibrary*> GetAllDMXLibraries();
+
+	/**  Return integer given an array of bytes. The first 4 bytes in the array will be used for the conversion. */
+	UFUNCTION(BlueprintCallable, Category = "DMX")
+	int32 BytesToInt(const TArray<uint8>& Bytes);
+
+	/**
+	 * Creates a literal UDMXEntityFixturePatch reference
+	 * @param	InFixturePatch	pointer to set the UDMXEntityFixturePatch to
+	 * @return	The literal UDMXEntityFixturePatch
+	 */
+	UFUNCTION(BlueprintPure, meta = (BlueprintInternalUseOnly = "TRUE"), Category = "DMX")
+	UDMXEntityFixturePatch* GetFixturePatch(FDMXEntityFixturePatchRef InFixturePatch);
+
+	/**
+	 * Gets a function map based on you active mode from FixturePatch
+	 * @param	InFixturePatch Selected Patch
+	 * @param	FDMXProtocolName Selected DMX protocol
+	 * @param	OutFunctionsMap Function and Channel value output map
+	 * @return	True if outputting was successfully
+	 */
+	UFUNCTION(BlueprintCallable, meta = (BlueprintInternalUseOnly = "TRUE", AutoCreateRefTerm = "SelectedProtocol"), Category = "DMX")
+	bool GetFunctionsMap(UDMXEntityFixturePatch* InFixturePatch, const FDMXProtocolName& SelectedProtocol, TMap<FName, int32>& OutFunctionsMap);
+
+	/**
+	 * Gets function channel value by input function name
+	 * @param	InName Looking fixture function name
+	 * @param	InFunctionsMap Function and Channel value input map
+	 * @return	Function channel value
+	 */
+	UFUNCTION(BlueprintCallable, meta = (BlueprintInternalUseOnly = "TRUE", AutoCreateRefTerm = "InName, InFunctionsMap"), Category = "DMX")
+	int32 GetFunctionsValue(const FName& InName, const TMap<FName, int32>& InFunctionsMap);
+
+	UPROPERTY(BlueprintAssignable, Category = "DMX")
+	FProtocolReceivedDelegate OnProtocolReceived;
+
+public:
+	//~ USubsystem interface begin
+	virtual void Initialize(FSubsystemCollectionBase& Collection) override;
+	//~ USubsystem interface end
+
+private:
+
+	// This function is a delegate for when protocols have input updates
+	UFUNCTION()
+	void BufferReceivedBroadcast(FName Protocol, uint16 UniverseID, const TArray<uint8>& Values);
+
+};
diff --git a/Engine/Plugins/Experimental/DigitalMultiplex/DMXEngine/Source/DMXRuntime/Public/Game/DMXComponent.h b/Engine/Plugins/Experimental/DigitalMultiplex/DMXEngine/Source/DMXRuntime/Public/Game/DMXComponent.h
new file mode 100644
--- /dev/null
+++ b/Engine/Plugins/Experimental/DigitalMultiplex/DMXEngine/Source/DMXRuntime/Public/Game/DMXComponent.h
@@ -0,0 +1,40 @@
+// Copyright Epic Games, Inc. All Rights Reserved.
+
+#pragma once
+
+#include "CoreMinimal.h"
+#include "Components/ActorComponent.h"
+#include "Library/DMXEntityReference.h"
+
+#include "DMXComponent.generated.h"
+
+class UDMXLibrary;
+
+UCLASS( ClassGroup=(DMX), meta=(BlueprintSpawnableComponent), HideCategories = ("Variable", "Sockets", "Tags", "Activation", "Cooking", "ComponentReplication", "AssetUserData", "Collision", "Events"))
+class DMXRUNTIME_API UDMXComponent
+	: public UActorComponent
+{
+	GENERATED_BODY()
+
+public:
+	UPROPERTY(EditAnywhere, Category = "DMX")
+	FDMXEntityFixturePatchRef FixturePatch;
+
+public:
+	UFUNCTION(BlueprintPure, Category = "DMX")
+	UDMXEntityFixturePatch* GetFixturePatch() const;
+
+	UFUNCTION(BlueprintCallable, Category = "DMX")
+	void SetFixturePatch(UDMXEntityFixturePatch* InFixturePatch);
+
+public:
+	// Sets default values for this component's properties
+	UDMXComponent();
+
+protected:
+	// Called when the game starts
+	virtual void BeginPlay() override;
+
+	// Called every frame
+	virtual void TickComponent(float DeltaTime, ELevelTick TickType, FActorComponentTickFunction* ThisTickFunction) override;
+};
diff --git a/Engine/Plugins/Experimental/DigitalMultiplex/DMXEngine/Source/DMXRuntime/Public/Library/DMXEntity.h b/Engine/Plugins/Experimental/DigitalMultiplex/DMXEngine/Source/DMXRuntime/Public/Library/DMXEntity.h
new file mode 100644
--- /dev/null
+++ b/Engine/Plugins/Experimental/DigitalMultiplex/DMXEngine/Source/DMXRuntime/Public/Library/DMXEntity.h
@@ -0,0 +1,101 @@
+// Copyright Epic Games, Inc. All Rights Reserved.
+
+#pragma once
+
+#include "Library/DMXObjectBase.h"
+#include "DMXProtocolTypes.h"
+
+#include "Misc/Guid.h"
+#include "Templates/SubclassOf.h"
+
+#include "DMXEntity.generated.h"
+
+class UDMXLibrary;
+
+/**  Base class for all entity types */
+UCLASS(abstract, meta = (DisplayName = "DMX Entity"))
+class DMXRUNTIME_API UDMXEntity
+	: public UDMXObjectBase
+{
+	GENERATED_BODY()
+
+public:
+	UPROPERTY(EditAnywhere, BlueprintReadOnly, Category = "Entity Properties", meta = (DisplayName = "Name", DisplayPriority = "1"))
+	FString Name;
+
+public:
+	UDMXEntity();
+
+	/**  Returns the entity name to be used in UI elements */
+	FString GetDisplayName() const;
+	/**  Updates this Entity's name and the UI friendly display name */
+	void SetName(const FString& InNewName);
+
+	/**
+	 * Checks for Entity correctness for usability with protocols.
+	 * @param OutReason	Reason for being invalid, if that's the case.
+	 * @return True if the Entity can be used with protocols.
+	 */
+	virtual bool IsValidEntity(FText& OutReason) const { return true; }
+
+	/**
+	 * Checks for Entity correctness for usability with protocols.
+	 * @return True if the Entity can be used with protocols.
+	 */
+	virtual bool IsValidEntity() const
+	{
+		FText OutReason = FText::GetEmpty();
+		return IsValidEntity(OutReason);
+	}
+
+	void SetParentLibrary(UDMXLibrary* InParent);
+	UDMXLibrary* GetParentLibrary() const { return ParentLibrary.Get(); }
+
+	/** This Entity's unique ID */
+	const FGuid& GetID() const { return Id; }
+	/** Used by DMX Library to resolve ID conflicts among entities */
+	void RefreshID();
+	/** Copy another Entity's ID. Used when copying, to not lose the original Entity's reference  */
+	void ReplicateID(UDMXEntity* Other);
+
+#if WITH_EDITOR
+	virtual void PostDuplicate(EDuplicateMode::Type DuplicateMode) override;
+#endif // WITH_EDITOR
+
+protected:
+	UPROPERTY()
+	TWeakObjectPtr<UDMXLibrary> ParentLibrary;
+	
+	/** Uniquely identifies the parameter, used for fixing up Blueprints that reference this Entity when renaming. */
+	UPROPERTY()
+	FGuid Id;
+};
+
+/**  Specialized version of UDMXEntity which represents an entity that is managed by a Universe */
+UCLASS()
+class DMXRUNTIME_API UDMXEntityUniverseManaged
+	: public UDMXEntity
+{
+	GENERATED_BODY()
+
+public:
+
+	UDMXEntityUniverseManaged();
+
+	UPROPERTY(EditAnywhere, Category = "Communication Properties", meta = (DisplayName = "Protocol", DisplayPriority = "1"))
+	FDMXProtocolName DeviceProtocol;
+
+	UPROPERTY(EditAnywhere, Category = "Universe Properties", meta = (ShowOnlyInnerProperties, DisplayPriority = 30))
+	TArray<FDMXUniverse> Universes;
+
+public:
+	virtual void PostLoad() override;
+
+#if WITH_EDITOR
+	virtual void PostEditChangeChainProperty(FPropertyChangedChainEvent& PropertyChangedEvent) override;
+#endif // WITH_EDITOR
+
+
+private:
+	void UpdateProtocolUniverses() const;
+};
diff --git a/Engine/Plugins/Experimental/DigitalMultiplex/DMXEngine/Source/DMXRuntime/Public/Library/DMXEntityController.h b/Engine/Plugins/Experimental/DigitalMultiplex/DMXEngine/Source/DMXRuntime/Public/Library/DMXEntityController.h
new file mode 100644
--- /dev/null
+++ b/Engine/Plugins/Experimental/DigitalMultiplex/DMXEngine/Source/DMXRuntime/Public/Library/DMXEntityController.h
@@ -0,0 +1,75 @@
+// Copyright Epic Games, Inc. All Rights Reserved.
+
+#pragma once
+
+#include "Library/DMXEntity.h"
+#include "DMXEntityController.generated.h"
+
+UCLASS(meta = (DisplayName = "DMX Controller"))
+class DMXRUNTIME_API UDMXEntityController
+	: public UDMXEntityUniverseManaged
+{
+	GENERATED_BODY()
+
+public:
+	/**  Defines where DMX data is sent to. */
+	UPROPERTY(EditAnywhere, Category = "DMX") // Hidden. For now it has just a single value. More values will be added in the future.
+	EDMXCommunicationTypes CommunicationMode;
+
+	/**  First Universe ID on this Controller's range. */
+	UPROPERTY(EditAnywhere, BlueprintReadOnly, Category = "Universe Properties", meta = (DisplayName = "Universe Start", DisplayPriority = 1, ClampMin = 1))
+	int32 UniverseLocalStart;
+
+	/**  Number of Universe IDs on this Controller's range, starting from Universe Start value. */
+	UPROPERTY(EditAnywhere, BlueprintReadOnly, Category = "Universe Properties", meta = (DisplayName = "Amount of Universes", DisplayPriority = 2, ClampMin = 1))
+	int32 UniverseLocalNum;
+
+	/**  Last Universe ID on this Controller's range, calculated from Universe Start and Amount of Universes. */
+	UPROPERTY(VisibleAnywhere, BlueprintReadOnly, Category = "Universe Properties", meta = (DisplayName = "Universe End", DisplayPriority = 3))
+	int32 UniverseLocalEnd;
+
+	/**
+	 * Offsets the Universe IDs range on this Controller before communication with other devices.
+	 * Useful to solve conflicts with Universe IDs from other devices on the same network.
+	 *
+	 * All other DMX Library settings use the normal Universe IDs range.
+	 * This allows the user to change all Universe IDs used by the Fixture Patches and
+	 * avoid conflicts with other devices by updating only the Controller's Remote Offset.
+	 */
+	UPROPERTY()
+	int32 RemoteOffset;
+
+	/**
+	 * First Universe ID on this Controller's range that is sent over the network.
+	 * Universe Start + Remote Offset
+	 */
+	UPROPERTY(EditAnywhere, BlueprintReadOnly, Category = "Universe Properties", meta = (DisplayName = "Remote Universe Range Start", DisplayPriority = 21))
+	int32 UniverseRemoteStart;
+
+	/**
+	 * Last Universe ID in this Controller's range that is sent over the network.
+	 * Universe End + Remote Offset
+	 */
+	UPROPERTY(VisibleAnywhere, BlueprintReadOnly, Category = "Universe Properties", meta = (DisplayName = "Remote Universe Range End", DisplayPriority = 22))
+	int32 UniverseRemoteEnd;
+
+public:
+	UDMXEntityController()
+		: UniverseLocalNum(1)
+	{
+		ValidateRangeValues();
+	}
+
+	/** Returns the currently assigned protocol for this controller */
+	UFUNCTION(BlueprintPure, Category = "DMX")
+	FName GetProtocol() const { return DeviceProtocol.Name; }
+
+#if WITH_EDITOR
+	virtual void PostEditChangeProperty(FPropertyChangedEvent& PropertyChangedEvent) override;
+#endif // WITH_EDITOR
+
+protected:
+	/** Keep range valid */
+	void ValidateRangeValues();
+	void UpdateUniversesFromRange();
+};
diff --git a/Engine/Plugins/Experimental/DigitalMultiplex/DMXEngine/Source/DMXRuntime/Public/Library/DMXEntityFader.h b/Engine/Plugins/Experimental/DigitalMultiplex/DMXEngine/Source/DMXRuntime/Public/Library/DMXEntityFader.h
new file mode 100644
--- /dev/null
+++ b/Engine/Plugins/Experimental/DigitalMultiplex/DMXEngine/Source/DMXRuntime/Public/Library/DMXEntityFader.h
@@ -0,0 +1,18 @@
+// Copyright Epic Games, Inc. All Rights Reserved.
+
+#pragma once
+
+#include "Library/DMXEntity.h"
+#include "DMXEntityFader.generated.h"
+
+UCLASS(meta = (DisplayName = "DMX Fader"))
+class DMXRUNTIME_API UDMXEntityFader
+	: public UDMXEntityUniverseManaged
+{
+	GENERATED_BODY()
+
+public:
+	UPROPERTY()
+	bool bIsActive;
+
+};
diff --git a/Engine/Plugins/Experimental/DigitalMultiplex/DMXEngine/Source/DMXRuntime/Public/Library/DMXEntityFixturePatch.h b/Engine/Plugins/Experimental/DigitalMultiplex/DMXEngine/Source/DMXRuntime/Public/Library/DMXEntityFixturePatch.h
new file mode 100644
--- /dev/null
+++ b/Engine/Plugins/Experimental/DigitalMultiplex/DMXEngine/Source/DMXRuntime/Public/Library/DMXEntityFixturePatch.h
@@ -0,0 +1,156 @@
+// Copyright Epic Games, Inc. All Rights Reserved.
+
+#pragma once
+
+#include "Library/DMXEntity.h"
+#include "Library/DMXEntityController.h"
+#include "Library/DMXEntityFixtureType.h"
+#include "DMXProtocolTypes.h"
+#include "DMXEntityFixturePatch.generated.h"
+
+class UDMXEntityFixtureType;
+
+UCLASS(BlueprintType, Blueprintable, meta = (DisplayName = "DMX Fixture Patch"))
+class DMXRUNTIME_API UDMXEntityFixturePatch
+	: public UDMXEntity
+{
+	GENERATED_BODY()
+
+public:
+	UPROPERTY(EditAnywhere, BlueprintReadOnly, Category = "Fixture Settings")
+	TArray<FName> CustomTags;
+
+	/**
+	 * If set to a value on a Controller's Universe IDs range (without the Range Offset),
+	 * this Patch's functions are sent over the network by that Controller.
+	 *
+	 * When set to a value on several Controllers' range, the functions are sent by all of those Controllers.
+	 */
+	UPROPERTY(EditAnywhere, BlueprintReadOnly, Category = "Fixture Settings", meta = (ClampMin = 0))
+	int32 UniverseID;
+
+	/** Auto-assign channel from drag/drop list order and available channels */
+	UPROPERTY(EditAnywhere, BlueprintReadOnly, Category = "Fixture Settings", meta = (DisplayName = "Auto-Assign Channel"))
+	bool bAutoAssignAddress;
+
+	/** Starting channel for when auto-assign channel is false */
+	UPROPERTY(EditAnywhere, Category = "Fixture Settings", meta = (EditCondition = "!bAutoAssignAddress", DisplayName = "Manual Starting Channel", UIMin = "1", UIMax = "512", ClampMin = "1", ClampMax = "512"))
+	int32 ManualStartingAddress;
+
+	/** Starting channel from auto-assignment. Used when AutoAssignAddress is true */
+	UPROPERTY(NonTransactional)
+	int32 AutoStartingAddress;
+
+	/** Property to point to the template parent fixture for details panel purposes */
+	UPROPERTY(VisibleDefaultsOnly, BlueprintReadOnly, Category = "Fixture Type Settings", meta = (DisplayName = "Fixture Type"))
+	UDMXEntityFixtureType* ParentFixtureTypeTemplate;
+
+	UPROPERTY(EditAnywhere, BlueprintReadOnly, Category = "Fixture Type Settings")
+	int32 ActiveMode;
+
+public:
+	/**
+	 * Returns the number of channels this Patch occupies with the Fixture functions from its Active Mode.
+	 * It'll always be at least 1 channel.
+	 */
+	UFUNCTION(BlueprintPure, Category = "DMX|Fixture Patch")
+	int32 GetChannelSpan() const;
+
+	/**  Return the active starting channel, evaluated after checking if Auto-Assignment is activated. */
+	UFUNCTION(BlueprintPure, Category = "DMX|Fixture Patch")
+	int32 GetStartingChannel() const;
+
+	/**
+	 * Return an array of function names for the currently active mode.
+	 * Functions outside the Active Mode's channel span range are ignored.
+	 */
+	UFUNCTION(BlueprintPure, Category = "DMX|Fixture Patch")
+	TArray<FName> GetAllFunctionsInActiveMode() const;
+
+	/**
+	 * Return map of function names and default values.
+	 * Functions outside the Active Mode's channel span range are ignored.
+	 */
+	UFUNCTION(BlueprintPure, Category = "DMX|Fixture Patch")
+	TMap<FName, int32> GetFunctionDefaultMap() const;
+
+	/**
+	 * Return map of function names and their assigned channels.
+	 * Functions outside the Active Mode's channel span range are ignored.
+	 */
+	UFUNCTION(BlueprintPure, Category = "DMX|Fixture Patch")
+	TMap<FName, int32> GetFunctionChannelAssignments() const;
+
+	/**
+	 * Return map of function names and their Data Types.
+	 * Functions outside the Active Mode's channel span range are ignored.
+	 */
+	UFUNCTION(BlueprintPure, Category = "DMX|Fixture Patch")
+	TMap<FName, EDMXFixtureSignalFormat> GetFunctionSignalFormats() const;
+
+	/**  Given a <Channel Index -> Raw Value> map , return map of function names and their values. */
+	UFUNCTION(BlueprintPure, Category = "DMX|Fixture Patch")
+	TMap<FName, int32> ConvertRawMapToFunctionMap(const TMap<int32, uint8>& RawMap) const;
+
+	/**
+	 * Return map of function channels and their values.
+	 * Functions outside the Active Mode's channel span range are ignored.
+	 */
+	UFUNCTION(BlueprintPure, Category = "DMX|Fixture Patch")
+	TMap<int32, uint8> ConvertFunctionMapToRawMap(const TMap<FName, int32>& FunctionMap) const;
+
+	/**  Return if given function map valid for this fixture. */
+	UFUNCTION(BlueprintPure, Category = "DMX|Fixture Patch")
+	bool IsMapValid(const TMap<FName, int32>& FunctionMap) const;
+
+	/**  Return if fixture contains function. */
+	UFUNCTION(BlueprintPure, Category = "DMX|Fixture Patch")
+	FORCEINLINE bool ContainsFunction(const FName& FunctionName) const
+	{
+		const TArray<FDMXFixtureFunction>& Functions = ParentFixtureTypeTemplate->Modes[ActiveMode].Functions;
+		return Functions.ContainsByPredicate([&FunctionName](const FDMXFixtureFunction& Function)
+			{
+				return FunctionName.IsEqual(*Function.FunctionName);
+			});
+	}
+
+	/**  Return a map that is valid for this fixture. */
+	UFUNCTION(BlueprintPure, Category = "DMX|Fixture Patch")
+	TMap<FName, int32> ConvertToValidMap(const TMap<FName, int32>& FunctionMap) const;
+
+	/**
+	 * Scans the parent DMXLibrary and returns the Controllers which Universe range
+	 * match this Patch's UniverseID
+	 */
+	UFUNCTION(BlueprintPure, Category = "DMX|Fixture Patch")
+	TArray<UDMXEntityController*> GetRelevantControllers() const;
+
+	/** Returns true if this Patch's UniverseID is in InController's range */
+	UFUNCTION(BlueprintPure, Category = "DMX|Fixture Patch")
+	FORCEINLINE bool IsInControllerRange(const UDMXEntityController* InController) const
+	{
+		return InController != nullptr &&
+			UniverseID >= InController->UniverseLocalStart && UniverseID <= InController->UniverseLocalEnd;
+	}
+
+	/** Returns true if this Patch's UniverseID is in any of InControllers' ranges */
+	UFUNCTION(BlueprintPure, Category = "DMX|Fixture Patch")
+	bool IsInControllersRange(const TArray<UDMXEntityController*>& InControllers) const;
+
+public:
+	UDMXEntityFixturePatch();
+
+	//~ Begin UDMXEntity Interface
+	virtual bool IsValidEntity(FText& OutReason) const;
+	//~ End UDMXEntity Interface
+
+	// Called from Fixture Type to keep ActiveMode in valid range when Modes are removed from the Type
+	void ValidateActiveMode();
+
+	FORCEINLINE bool CanReadActiveMode() const
+	{
+		return ParentFixtureTypeTemplate != nullptr
+			&& ParentFixtureTypeTemplate->Modes.Num() > 0
+			&& ActiveMode < ParentFixtureTypeTemplate->Modes.Num();
+	}
+};
diff --git a/Engine/Plugins/Experimental/DigitalMultiplex/DMXEngine/Source/DMXRuntime/Public/Library/DMXEntityFixtureType.h b/Engine/Plugins/Experimental/DigitalMultiplex/DMXEngine/Source/DMXRuntime/Public/Library/DMXEntityFixtureType.h
new file mode 100644
--- /dev/null
+++ b/Engine/Plugins/Experimental/DigitalMultiplex/DMXEngine/Source/DMXRuntime/Public/Library/DMXEntityFixtureType.h
@@ -0,0 +1,142 @@
+// Copyright Epic Games, Inc. All Rights Reserved.
+
+#pragma once
+
+#include "Library/DMXEntity.h"
+#include "DMXProtocolTypes.h"
+
+#include "DMXEntityFixtureType.generated.h"
+
+class UDMXImport;
+
+USTRUCT(BlueprintType)
+struct DMXRUNTIME_API FDMXFixtureSubFunction
+{
+	GENERATED_BODY()
+
+	UPROPERTY(EditAnywhere, BlueprintReadOnly, Category = "DMX")
+	FString FunctionName;
+
+	/** Minimum value in the range of values that represent this sub function */
+	UPROPERTY(EditAnywhere, BlueprintReadOnly, Category = "DMX")
+	uint8 MinValue;
+
+	/** Maximum value in the range of values that represent this sub function */
+	UPROPERTY(EditAnywhere, BlueprintReadOnly, Category = "DMX")
+	uint8 MaxValue;
+};
+
+USTRUCT(BlueprintType)
+struct DMXRUNTIME_API FDMXFixtureFunction
+{
+	GENERATED_BODY()
+
+	UPROPERTY(EditAnywhere, BlueprintReadOnly, meta = (DisplayPriority = "10"), Category = "DMX")
+	FString FunctionName;
+
+	UPROPERTY(EditAnywhere, BlueprintReadOnly, meta = (DisplayPriority = "20"), Category = "DMX")
+	FString Description;
+
+	/** Ranges of values that each represent one sub function */
+	UPROPERTY(EditAnywhere, BlueprintReadOnly, meta = (DisplayPriority = "40"), Category = "DMX")
+	TArray<FDMXFixtureSubFunction> SubFunctions;
+
+	/** Initial value for this function when no value is set */
+	UPROPERTY(EditAnywhere, BlueprintReadOnly, meta = (DisplayPriority = "30"), Category = "DMX")
+	int64 DefaultValue;
+
+	/** This function's starting channel */
+	UPROPERTY(VisibleAnywhere, BlueprintReadOnly, meta = (DisplayName = "Channel Assignment", ClampMin = "1", ClampMax = "512", DisplayPriority = "1"), Category = "DMX")
+	int32 Channel;
+
+	/**
+	 * This function's channel offset.
+	 * E.g.: if the function's starting channel is supposed to be 10
+	 * and ChannelOffset = 5, the function's starting channel becomes 15
+	 * and all following functions follow it accordingly.
+	 */
+	UPROPERTY(EditAnywhere, BlueprintReadOnly, meta = (DisplayName = "Channel Offset", ClampMin = "0", ClampMax = "511", DisplayPriority = "1"), Category = "DMX")
+	int32 ChannelOffset;
+
+	/** This function's data type. Defines the used number of channels (bytes) */
+	UPROPERTY(EditAnywhere, BlueprintReadOnly, meta = (DisplayPriority = "5"), Category = "DMX")
+	EDMXFixtureSignalFormat DataType;
+
+	FDMXFixtureFunction()
+		: FunctionName()
+		, Description()
+		, SubFunctions()
+		, DefaultValue(0)
+		, Channel(1)
+		, ChannelOffset(0)
+		, DataType(EDMXFixtureSignalFormat::E8Bit)
+	{}
+};
+
+USTRUCT(BlueprintType)
+struct DMXRUNTIME_API FDMXFixtureMode
+{
+	GENERATED_BODY()
+
+	UPROPERTY(EditAnywhere, BlueprintReadOnly, meta = (DisplayPriority = "1"), Category = "DMX")
+	FString ModeName;
+
+	UPROPERTY(EditAnywhere, BlueprintReadOnly, meta = (DisplayPriority = "15"), Category = "DMX")
+	TArray<FDMXFixtureFunction> Functions;
+
+	/** Number of channels (bytes) used by this mode's functions */
+	UPROPERTY(EditAnywhere, BlueprintReadOnly, meta = (ClampMin = "1", ClampMax = "512", DisplayPriority = "10", EditCondition = "!bAutoChannelSpan"), Category = "DMX")
+	int32 ChannelSpan;
+
+	/**
+	 * When enabled, ChannelSpan is automatically set based on the created functions and their data types.
+	 * If disabled, ChannelSpan can be manually set and functions and functions' channels beyond the
+	 * specified span will be ignored.
+	 */
+	UPROPERTY(EditAnywhere, meta = (DisplayPriority = "5"), Category = "DMX")
+	bool bAutoChannelSpan;
+
+	FDMXFixtureMode()
+		: ChannelSpan(1)
+		, bAutoChannelSpan(true)
+	{}
+};
+
+UCLASS(BlueprintType, Blueprintable, meta = (DisplayName = "DMX Fixture Type"))
+class DMXRUNTIME_API UDMXEntityFixtureType
+	: public UDMXEntity
+{
+	GENERATED_BODY()
+
+public:
+	UPROPERTY(EditAnywhere, BlueprintReadOnly, Category = "Fixture Settings")
+	UDMXImport* DMXImport;
+
+	UPROPERTY(EditAnywhere, BlueprintReadOnly, Category = "Fixture Settings", meta = (DisplayName = "DMX Category"))
+	FDMXFixtureCategory DMXCategory;
+
+	UPROPERTY(EditAnywhere, BlueprintReadOnly, Category = "Fixture Settings")
+	TArray<FDMXFixtureMode> Modes;
+
+public:
+	UFUNCTION(BlueprintCallable, Category = "Fixture Settings")
+	void SetModesFromDMXImport(UDMXImport* DMXImportAsset);
+
+	static void SetFunctionSize(FDMXFixtureMode& InMode, FDMXFixtureFunction& InFunction, uint8 Size);
+
+	/** Gets the last channel occupied by the Function */
+	static uint8 GetFunctionLastChannel(const FDMXFixtureFunction& Function);
+
+	static void ClampDefaultValue(FDMXFixtureFunction& InFunction);
+
+	static uint8 NumChannelsToOccupy(EDMXFixtureSignalFormat DataType);
+
+	static uint32 ClampValueToDataType(EDMXFixtureSignalFormat DataType, uint32 InValue);
+
+#if WITH_EDITOR
+	virtual void PostEditChangeChainProperty(FPropertyChangedChainEvent& PropertyChangedEvent) override;
+	virtual void PostEditUndo() override;
+
+	static void UpdateModeChannelProperties(FDMXFixtureMode& Mode);
+#endif // WITH_EDITOR
+};
diff --git a/Engine/Plugins/Experimental/DigitalMultiplex/DMXEngine/Source/DMXRuntime/Public/Library/DMXImport.h b/Engine/Plugins/Experimental/DigitalMultiplex/DMXEngine/Source/DMXRuntime/Public/Library/DMXImport.h
new file mode 100644
--- /dev/null
+++ b/Engine/Plugins/Experimental/DigitalMultiplex/DMXEngine/Source/DMXRuntime/Public/Library/DMXImport.h
@@ -0,0 +1,138 @@
+// Copyright Epic Games, Inc. All Rights Reserved.
+
+#pragma once
+
+#include "CoreMinimal.h"
+#include "UObject/Object.h"
+#include "DMXImport.generated.h"
+
+namespace DMXImport
+{
+    template <typename EnumType> static EnumType GetEnumValueFromString(const FString& String)
+    {
+        return static_cast<EnumType>(StaticEnum<EnumType>()->GetValueByName(FName(*String)));
+    }
+
+    DMXRUNTIME_API FDMXColorCIE ParseColorCIE(const FString& InColor);
+
+    DMXRUNTIME_API FMatrix ParseMatrix(FString&& InMatrixStr);
+};
+
+USTRUCT(BlueprintType)
+struct FDMXColorCIE
+{
+    GENERATED_BODY()
+
+        FDMXColorCIE()
+        : X(0.f)
+        , Y(0.f)
+        , YY(0)
+    {}
+
+    UPROPERTY(EditAnywhere, Category = Color, meta = (ClampMin = "0", ClampMax = "1.0"))
+    float X;
+
+    UPROPERTY(EditAnywhere, Category = Color, meta = (ClampMin = "0", ClampMax = "1.0"))
+    float Y;
+
+    UPROPERTY(EditAnywhere, Category = Color, meta = (ClampMin = "0", ClampMax = "255"))
+    uint8 YY;
+};
+
+UCLASS(BlueprintType, Blueprintable)
+class DMXRUNTIME_API UDMXImportFixtureType
+    : public UObject
+{
+    GENERATED_BODY()
+};
+
+UCLASS(BlueprintType, Blueprintable)
+class DMXRUNTIME_API UDMXImportAttributeDefinitions
+    : public UObject
+{
+    GENERATED_BODY()
+};
+
+UCLASS(BlueprintType, Blueprintable)
+class DMXRUNTIME_API UDMXImportWheels
+    : public UObject
+{
+    GENERATED_BODY()
+};
+
+UCLASS(BlueprintType, Blueprintable)
+class DMXRUNTIME_API UDMXImportPhysicalDescriptions
+    : public UObject
+{
+    GENERATED_BODY()
+};
+
+UCLASS(BlueprintType, Blueprintable)
+class DMXRUNTIME_API UDMXImportModels
+    : public UObject
+{
+    GENERATED_BODY()
+};
+
+UCLASS(BlueprintType, Blueprintable)
+class DMXRUNTIME_API UDMXImportGeometries
+    : public UObject
+{
+    GENERATED_BODY()
+};
+
+UCLASS(BlueprintType, Blueprintable)
+class DMXRUNTIME_API UDMXImportDMXModes
+    : public UObject
+{
+    GENERATED_BODY()
+
+};
+
+UCLASS(BlueprintType, Blueprintable)
+class DMXRUNTIME_API UDMXImportProtocols
+    : public UObject
+{
+    GENERATED_BODY()
+
+};
+
+UCLASS(BlueprintType, Blueprintable, abstract)
+class DMXRUNTIME_API UDMXImport
+	: public UObject
+{
+	GENERATED_BODY()
+
+public:
+	template <typename TType>
+	TType* CreateNewObject()
+	{
+        FName NewName = MakeUniqueObjectName(this, TType::StaticClass());
+        return NewObject<TType>(this, TType::StaticClass(), NewName, RF_Public);
+	}
+
+public:
+	UPROPERTY(VisibleAnywhere, Category = "Fixture Type", meta = (EditInline))
+	UDMXImportFixtureType* FixtureType;
+
+    UPROPERTY(VisibleAnywhere, Category = "Attribute Definitions", meta = (EditInline))
+    UDMXImportAttributeDefinitions* AttributeDefinitions;
+
+    UPROPERTY(VisibleAnywhere, Category = "Wheels", meta = (EditInline))
+    UDMXImportWheels* Wheels;
+
+    UPROPERTY(VisibleAnywhere, Category = "Physical Descriptions", meta = (EditInline))
+    UDMXImportPhysicalDescriptions* PhysicalDescriptions;
+
+    UPROPERTY(VisibleAnywhere, Category = "Models", meta = (EditInline))
+    UDMXImportModels* Models;
+
+    UPROPERTY(VisibleAnywhere, Category = "Geometries", meta = (EditInline))
+    UDMXImportGeometries* Geometries;
+
+    UPROPERTY(VisibleAnywhere, Category = "DMXModes", meta = (EditInline))
+    UDMXImportDMXModes* DMXModes;
+
+    UPROPERTY(VisibleAnywhere, Category = "Protocols", meta = (EditInline))
+    UDMXImportProtocols* Protocols;
+};
diff --git a/Engine/Plugins/Experimental/DigitalMultiplex/DMXProtocol/DMXProtocol.uplugin b/Engine/Plugins/Experimental/DigitalMultiplex/DMXProtocol/DMXProtocol.uplugin
--- a/Engine/Plugins/Experimental/DigitalMultiplex/DMXProtocol/DMXProtocol.uplugin
+++ b/Engine/Plugins/Experimental/DigitalMultiplex/DMXProtocol/DMXProtocol.uplugin
@@ -10,23 +10,44 @@
   "DocsURL": "",
   "MarketplaceURL": "",
   "SupportURL": "",
-  "CanContainContent": true,
+  "CanContainContent": false,
   "IsBetaVersion": true,
   "Modules": [
     {
       "Name": "DMXProtocol",
       "Type": "Runtime",
-      "LoadingPhase": "Default"
+      "LoadingPhase": "PreDefault",
+      "WhitelistPlatforms": 
+      [
+        "Win64",
+        "Win32",
+        "Mac",
+        "Linux"
+      ]
     },
     {
       "Name": "DMXProtocolArtNet",
       "Type": "Runtime",
-      "LoadingPhase": "Default"
+      "LoadingPhase": "Default",
+      "WhitelistPlatforms": 
+      [
+        "Win64",
+        "Win32",
+        "Mac",
+        "Linux"
+      ]
     },
     {
       "Name": "DMXProtocolSACN",
       "Type": "Runtime",
-      "LoadingPhase": "Default"
+      "LoadingPhase": "Default",
+      "WhitelistPlatforms": 
+      [
+        "Win64",
+        "Win32",
+        "Mac",
+        "Linux"
+      ]
     }
   ]
 }
diff --git a/Engine/Plugins/Experimental/DigitalMultiplex/DMXProtocol/Source/DMXProtocol/DMXProtocol.Build.cs b/Engine/Plugins/Experimental/DigitalMultiplex/DMXProtocol/Source/DMXProtocol/DMXProtocol.Build.cs
--- a/Engine/Plugins/Experimental/DigitalMultiplex/DMXProtocol/Source/DMXProtocol/DMXProtocol.Build.cs
+++ b/Engine/Plugins/Experimental/DigitalMultiplex/DMXProtocol/Source/DMXProtocol/DMXProtocol.Build.cs
@@ -12,7 +12,7 @@ public class DMXProtocol : ModuleRules
 			new string[]
 			{
 				"Core",
-            }
+			}
 			);
 			
 		
@@ -21,12 +21,25 @@ public class DMXProtocol : ModuleRules
 			{
 				"CoreUObject",
 				"Engine",
-                "Serialization",
-                "Networking",
-                "Sockets",
-                "Json"
-            }
-			);
-		
+				"Serialization",
+				"Networking",
+				"Sockets",
+				"Json"
+			}
+		);
+
+		if (Target.Type == TargetType.Editor)
+		{
+			DynamicallyLoadedModuleNames.AddRange(
+				new string[] {
+					"Settings",
+				});
+
+			PrivateIncludePathModuleNames.AddRange(
+				new string[] {
+					"Settings",
+				});
+		}
+
 	}
 }
diff --git a/Engine/Plugins/Experimental/DigitalMultiplex/DMXProtocol/Source/DMXProtocol/Private/DMXProtocolModule.cpp b/Engine/Plugins/Experimental/DigitalMultiplex/DMXProtocol/Source/DMXProtocol/Private/DMXProtocolModule.cpp
--- a/Engine/Plugins/Experimental/DigitalMultiplex/DMXProtocol/Source/DMXProtocol/Private/DMXProtocolModule.cpp
+++ b/Engine/Plugins/Experimental/DigitalMultiplex/DMXProtocol/Source/DMXProtocol/Private/DMXProtocolModule.cpp
@@ -3,16 +3,59 @@
 #include "DMXProtocolModule.h"
 #include "Interfaces/IDMXProtocol.h"
 #include "Interfaces/IDMXProtocolFactory.h"
+#include "DMXProtocolSettings.h"
+
+#if WITH_EDITOR
+#include "ISettingsModule.h"
+#endif
 
 IMPLEMENT_MODULE( FDMXProtocolModule, DMXProtocol );
+DEFINE_LOG_CATEGORY(LogDMXProtocol);
 
 #define LOCTEXT_NAMESPACE "DMXProtocolModule"
 
+/**  IDMXProtocol.h static declarations */
+FOnNetworkInterfaceChanged IDMXProtocol::OnNetworkInterfaceChanged;
+
 const TCHAR* FDMXProtocolModule::BaseModuleName = TEXT("DMXProtocol");
 
+const TMap<FName, TSharedPtr<IDMXProtocol>>& FDMXProtocolModule::GetProtocols() const
+{
+	return DMXProtocols;
+}
+
+void FDMXProtocolModule::StartupModule()
+{
+
+#if WITH_EDITOR
+	ISettingsModule* SettingsModule = FModuleManager::GetModulePtr<ISettingsModule>("Settings");
+
+	// Register DMX Protocol global settings
+	if (SettingsModule != nullptr)
+	{	
+		SettingsModule->RegisterSettings("Project", "Plugins", "DMX Plugin",
+			LOCTEXT("ProjectSettings_Label", "DMX Plugin"),
+			LOCTEXT("ProjectSettings_Description", "Configure DMX plugin global settings"),
+			GetMutableDefault<UDMXProtocolSettings>()
+		);
+	}
+#endif // WITH_EDITOR
+
+}
+
 void FDMXProtocolModule::ShutdownModule()
 {
 	ShutdownAllDMXProtocols();
+
+#if WITH_EDITOR
+	// Unregister DMX Protocol global settings
+	ISettingsModule* SettingsModule = FModuleManager::GetModulePtr<ISettingsModule>("Settings");
+	if (SettingsModule)
+	{
+		SettingsModule->UnregisterSettings("Project", "Plugins", "DMX Protocol");
+	}
+#endif // WITH_EDITOR
+
 }
 
 static FName GetProtocolModuleName(const FString& ProtocolName)
@@ -46,7 +89,7 @@ void FDMXProtocolModule::RegisterProtocol(const FName& FactoryName, IDMXProtocol
 	}
 	else
 	{
-		UE_LOG_DMXPROTOCOL(Warning, TEXT("Trying to add existing protocol %s"), *FactoryName.ToString());
+		UE_LOG_DMXPROTOCOL(Verbose, TEXT("Trying to add existing protocol %s"), *FactoryName.ToString());
 	}
 }
 
@@ -59,7 +102,7 @@ void FDMXProtocolModule::UnregisterProtocol(const FName& FactoryName)
 	}
 	else
 	{
-		UE_LOG_DMXPROTOCOL(Warning, TEXT("Trying to remove unexisting protocol %s"), *FactoryName.ToString());
+		UE_LOG_DMXPROTOCOL(Verbose, TEXT("Trying to remove unexisting protocol %s"), *FactoryName.ToString());
 	}
 }
 
@@ -69,7 +112,7 @@ FDMXProtocolModule& FDMXProtocolModule::Get()
 }
 
 
-IDMXProtocol* FDMXProtocolModule::GetProtocol(const FName InProtocolName)
+TSharedPtr<IDMXProtocol> FDMXProtocolModule::GetProtocol(const FName InProtocolName)
 {
 	TSharedPtr<IDMXProtocol>* DMXProtocolPtr = nullptr;
 	if (!InProtocolName.IsNone())
@@ -94,7 +137,7 @@ IDMXProtocol* FDMXProtocolModule::GetProtocol(const FName InProtocolName)
 					bool* bNotedPreviously = DMXProtocolFailureNotes.Find(InProtocolName);
 					if (!bNotedPreviously || !(*bNotedPreviously))
 					{
-						UE_LOG_DMXPROTOCOL(Warning, TEXT("Unable to create Protocol %s"), *InProtocolName.ToString());
+						UE_LOG_DMXPROTOCOL(Verbose, TEXT("Unable to create Protocol %s"), *InProtocolName.ToString());
 						DMXProtocolFailureNotes.Add(InProtocolName, true);
 					}
 				}
@@ -102,7 +145,7 @@ IDMXProtocol* FDMXProtocolModule::GetProtocol(const FName InProtocolName)
 		}
 	}
 
-	return (DMXProtocolPtr == nullptr) ? nullptr : (*DMXProtocolPtr).Get();
+	return (DMXProtocolPtr == nullptr) ? nullptr : *DMXProtocolPtr;
 }
 
 const TMap<FName, IDMXProtocolFactory*>& FDMXProtocolModule::GetProtocolFactories() const
@@ -122,7 +165,7 @@ void FDMXProtocolModule::ShutdownDMXProtocol(const FName& ProtocolName)
 		}
 		else
 		{
-			UE_LOG_DMXPROTOCOL(Warning, TEXT("DMXProtocol instance %s not found, unable to destroy."), *ProtocolName.ToString());
+			UE_LOG_DMXPROTOCOL(Verbose, TEXT("DMXProtocol instance %s not found, unable to destroy."), *ProtocolName.ToString());
 		}
 	}
 }
diff --git a/Engine/Plugins/Experimental/DigitalMultiplex/DMXProtocol/Source/DMXProtocol/Private/DMXProtocolSettings.cpp b/Engine/Plugins/Experimental/DigitalMultiplex/DMXProtocol/Source/DMXProtocol/Private/DMXProtocolSettings.cpp
new file mode 100644
--- /dev/null
+++ b/Engine/Plugins/Experimental/DigitalMultiplex/DMXProtocol/Source/DMXProtocol/Private/DMXProtocolSettings.cpp
@@ -0,0 +1,49 @@
+// Copyright Epic Games, Inc. All Rights Reserved.
+
+#include "DMXProtocolSettings.h"
+#include "Interfaces/IDMXProtocol.h"
+
+UDMXProtocolSettings::UDMXProtocolSettings()
+	: InterfaceIPAddress(TEXT("0.0.0.0"))
+	, bShouldUseUnicast(false)
+{
+	FixtureCategories =
+	{
+		TEXT("Static"),
+		TEXT("Matrix/Pixel Bar"),
+		TEXT("Moving Head"),
+		TEXT("Moving Mirror"),
+		TEXT("Strobe"),
+		TEXT("Other")
+	};
+}
+
+#if WITH_EDITOR
+
+void UDMXProtocolSettings::PostEditChangeProperty(FPropertyChangedEvent& PropertyChangedEvent)
+{
+	FName PropertyName;
+	if (PropertyChangedEvent.Property)
+	{
+		PropertyName = PropertyChangedEvent.Property->GetFName();
+	}
+
+	if (PropertyName == GET_MEMBER_NAME_CHECKED(UDMXProtocolSettings, InterfaceIPAddress) || PropertyName == GET_MEMBER_NAME_CHECKED(UDMXProtocolSettings, UnicastEndpoint))
+	{
+		IDMXProtocol::OnNetworkInterfaceChanged.Broadcast(InterfaceIPAddress);
+	}
+
+	if (PropertyName == GET_MEMBER_NAME_CHECKED(UDMXProtocolSettings, FixtureCategories))
+	{
+		if (FixtureCategories.Num() == 0)
+		{
+			FixtureCategories.Add(TEXT("Other"));
+		}
+
+		FDMXFixtureCategory::OnPossibleValuesUpdated.Broadcast();
+	}
+
+	Super::PostEditChangeProperty(PropertyChangedEvent);
+}
+
+#endif // WITH_EDITOR
\ No newline at end of file
diff --git a/Engine/Plugins/Experimental/DigitalMultiplex/DMXProtocol/Source/DMXProtocol/Private/DMXProtocolTypes.cpp b/Engine/Plugins/Experimental/DigitalMultiplex/DMXProtocol/Source/DMXProtocol/Private/DMXProtocolTypes.cpp
--- a/Engine/Plugins/Experimental/DigitalMultiplex/DMXProtocol/Source/DMXProtocol/Private/DMXProtocolTypes.cpp
+++ b/Engine/Plugins/Experimental/DigitalMultiplex/DMXProtocol/Source/DMXProtocol/Private/DMXProtocolTypes.cpp
@@ -1,16 +1,121 @@
 // Copyright Epic Games, Inc. All Rights Reserved.
 
 #include "DMXProtocolTypes.h"
+#include "Interfaces/IDMXProtocol.h"
+#include "DMXProtocolSettings.h"
 
-bool FDMXBuffer::SetDMXFragment(const IDMXFragmentMap & DMXFragment)
+TArray<FName> FDMXProtocolName::GetPossibleValues()
 {
-	for (TMap<uint32, uint8>::TConstIterator It = DMXFragment.CreateConstIterator(); It; ++It)
+	return IDMXProtocol::GetProtocolNames();
+}
+
+FDMXProtocolName::FDMXProtocolName(TSharedPtr<IDMXProtocol> InProtocol)
+{
+	if (InProtocol.IsValid())
+	{
+		Name = InProtocol->GetProtocolName();
+	}
+}
+
+FDMXProtocolName::FDMXProtocolName(const FName& InName)
+	: Name(InName)
+{}
+
+TSharedPtr<IDMXProtocol> FDMXProtocolName::GetProtocol() const
+{
+	if (Name.IsNone())
+	{
+		return nullptr; 
+	}
+	return IDMXProtocol::Get(Name);
+}
+
+bool FDMXProtocolName::IsValid() const
+{
+	return GetProtocol().IsValid();
+}
+
+FSimpleMulticastDelegate FDMXFixtureCategory::OnPossibleValuesUpdated;
+
+TArray<FName> FDMXFixtureCategory::GetPossibleValues()
+{
+	return GetDefault<UDMXProtocolSettings>()->FixtureCategories.Array();
+}
+
+FName FDMXFixtureCategory::GetFirstValue()
+{
+	const TSet<FName>& FixtureCategories = GetDefault<UDMXProtocolSettings>()->FixtureCategories;
+
+	for (const auto& Itt : FixtureCategories)
+	{
+		return Itt;
+	}
+
+	return FName();
+}
+
+FDMXFixtureCategory::FDMXFixtureCategory()
+{
+	Name = GetFirstValue();
+}
+
+FDMXFixtureCategory::FDMXFixtureCategory(const FName& InName)
+	: Name(InName)
+{}
+
+FString UDMXNameContainersConversions::Conv_DMXProtocolNameToString(const FDMXProtocolName & InProtocolName)
+{
+	return InProtocolName.Name.ToString();
+}
+
+FName UDMXNameContainersConversions::Conv_DMXProtocolNameToName(const FDMXProtocolName & InProtocolName)
+{
+	return InProtocolName.Name;
+}
+
+FString UDMXNameContainersConversions::Conv_DMXFixtureCategoryToString(const FDMXFixtureCategory & InFixtureCategory)
+{
+	return InFixtureCategory.Name.ToString();
+}
+
+FName UDMXNameContainersConversions::Conv_DMXFixtureCategoryToName(const FDMXFixtureCategory & InFixtureCategory)
+{
+	return InFixtureCategory.Name;
+}
+
+bool FDMXBuffer::SetDMXFragment(const IDMXFragmentMap & InDMXFragment)
+{
+	for (const TPair<uint32, uint8>& It : InDMXFragment)
 	{
-		if (It->Key < (DMX_UNIVERSE_SIZE - 1))
+		if (It.Key <= (DMX_UNIVERSE_SIZE) && It.Key > 0)
+		{
+			DMXData[It.Key-1] = It.Value;
+		}
+		else
 		{
-			DMXData[It->Key] = It->Value;
+			return false;
 		}
 	}
 
-	return false;
+	// Increase Sequence
+	SequanceID++;
+
+	return true;
+}
+
+bool FDMXBuffer::SetDMXBuffer(const uint8* InBuffer, uint32 InSize)
+{
+	if (InSize <= (DMX_UNIVERSE_SIZE) && InSize > 0)
+	{
+		FMemory::Memcpy(DMXData.GetData(), InBuffer, InSize);
+	}
+	else
+	{
+		return false;
+	}
+
+	// Increase Sequence
+	SequanceID++;
+
+	return true;
 }
diff --git a/Engine/Plugins/Experimental/DigitalMultiplex/DMXProtocol/Source/DMXProtocol/Private/Tests/DMXProtocolTests.cpp b/Engine/Plugins/Experimental/DigitalMultiplex/DMXProtocol/Source/DMXProtocol/Private/Tests/DMXProtocolTests.cpp
--- a/Engine/Plugins/Experimental/DigitalMultiplex/DMXProtocol/Source/DMXProtocol/Private/Tests/DMXProtocolTests.cpp
+++ b/Engine/Plugins/Experimental/DigitalMultiplex/DMXProtocol/Source/DMXProtocol/Private/Tests/DMXProtocolTests.cpp
@@ -6,11 +6,11 @@
 #include "DMXProtocolCommon.h"
 #include "Dom/JsonObject.h"
 
-#include "DMXProtocol.h"
-
+#include "Interfaces/IDMXProtocol.h"
+#include "DMXProtocolTypes.h"
 
 class FDMXProtocolTest
-	: public FDMXProtocol
+	: public IDMXProtocol
 {
 public:
 	//~ Begin IDMXProtocolBase implementation
@@ -20,11 +20,22 @@ public:
 	//~ End IDMXProtocolBase implementation
 
 	//~ Begin IDMXProtocol implementation
+	virtual const FName& GetProtocolName() const override { return ProtocolName;  }
+	virtual TSharedPtr<FJsonObject> GetSettings() const override { return Settings; }
 	virtual TSharedPtr<IDMXProtocolSender> GetSenderInterface() const override { return nullptr; }
-	virtual FName GetProtocolName() const override { return ProtocolName; }
-	virtual void Reload() override {}
-	virtual bool SendDMX(uint16 UniverseID, uint8 PortID, const TSharedPtr<FDMXBuffer>& DMXBuffer) const override { return true; }
+	virtual EDMXSendResult SendDMXFragment(uint16 UniverseID, const IDMXFragmentMap& DMXFragment) override { return EDMXSendResult::Success; }
+	virtual EDMXSendResult SendDMXFragmentCreate(uint16 InUniverseID, const IDMXFragmentMap& DMXFragment) override { return EDMXSendResult::Success; }
+	virtual uint16 GetFinalSendUniverseID(uint16 InUniverseID) const override { return InUniverseID; }
 	virtual bool IsEnabled() const override { return true; }
+	virtual TSharedPtr<IDMXProtocolUniverse, ESPMode::ThreadSafe> AddUniverse(const FJsonObject& InSettings) override { return nullptr; }
+	virtual void CollectUniverses(const TArray<FDMXUniverse>& Universes) override {}
+	virtual bool RemoveUniverseById(uint32 InUniverseId) override { return true; }
+	virtual void RemoveAllUniverses() override { }
+	virtual TSharedPtr<IDMXProtocolUniverse, ESPMode::ThreadSafe> GetUniverseById(uint32 InUniverseId) const override { return nullptr; }
+	virtual uint32 GetUniversesNum() const override { return 0; }
+	virtual uint16 GetMinUniverseID() const override { return 0; }
+	virtual uint16 GetMaxUniverses() const override { return 1; }
+	virtual FOnUniverseInputUpdateEvent& GetOnUniverseInputUpdate() override { return OnUniverseInputUpdateEvent; }
 	//~ End IDMXProtocol implementation
 
 	//~ Begin IDMXProtocolRDM implementation
@@ -35,12 +46,15 @@ public:
 	//~ Only the factory makes instances
 	FDMXProtocolTest() = delete;
 	explicit FDMXProtocolTest(FName InProtocolName, FJsonObject& InSettings)
-		: FDMXProtocol(InProtocolName, InSettings)
-		, ProtocolName(InProtocolName)
-	{}
+		: ProtocolName(InProtocolName)
+	{
+		Settings = MakeShared<FJsonObject>(InSettings);
+	}
 
 private:
 	FName ProtocolName;
+	TSharedPtr<FJsonObject> Settings;
+	FOnUniverseInputUpdateEvent OnUniverseInputUpdateEvent;
 };
 
 
@@ -57,14 +71,14 @@ public:
 		{
 			if (!ProtocolArtNetPtr->Init())
 			{
-				UE_LOG_DMXPROTOCOL(Warning, TEXT("TEST Protocol failed to initialize!"));
+				UE_LOG_DMXPROTOCOL(Verbose, TEXT("TEST Protocol failed to initialize!"));
 				ProtocolArtNetPtr->Shutdown();
 				ProtocolArtNetPtr = nullptr;
 			}
 		}
 		else
 		{
-			UE_LOG_DMXPROTOCOL(Warning, TEXT("TEST Protocol disabled!"));
+			UE_LOG_DMXPROTOCOL(Verbose, TEXT("TEST Protocol disabled!"));
 			ProtocolArtNetPtr->Shutdown();
 			ProtocolArtNetPtr = nullptr;
 		}
@@ -125,7 +139,7 @@ bool FDMXProtocolFactoryTest::RunTest(const FString& Parameters)
 	FDMXProtocolModule& DMXProtocolModule = FModuleManager::GetModuleChecked<FDMXProtocolModule>("DMXProtocol");
 
 	// Store the protocol pointer
-	IDMXProtocol* CachedProtocol = nullptr;
+	TSharedPtr<IDMXProtocol> CachedProtocol = nullptr;
 
 	// Try to register 3 times
 	TArray<TUniquePtr<IDMXProtocolFactory>> Factories;
@@ -141,7 +155,7 @@ bool FDMXProtocolFactoryTest::RunTest(const FString& Parameters)
 			CachedProtocol = IDMXProtocol::Get(ProtocolName);
 		}
 
-		TestNotNull(TEXT("Protocol should exists"), IDMXProtocol::Get(ProtocolName));
+		TestTrue(TEXT("Protocol should exists"), IDMXProtocol::Get(ProtocolName).IsValid());
 		TestEqual(TEXT("Should return same protocol instance"), CachedProtocol, IDMXProtocol::Get(ProtocolName));
 	}
 	Factories.Empty();
@@ -151,7 +165,7 @@ bool FDMXProtocolFactoryTest::RunTest(const FString& Parameters)
 		TUniquePtr<IDMXProtocolFactory> Factory = MakeUnique<FDMXProtocolFactoryTestFactory>();
 		DMXProtocolModule.RegisterProtocol(ProtocolName, Factory.Get());
 		DMXProtocolModule.UnregisterProtocol(ProtocolName);
-		TestNull(TEXT("Protocol should not exists"), IDMXProtocol::Get(ProtocolName));
+		TestFalse(TEXT("Protocol should not exists"), IDMXProtocol::Get(ProtocolName).IsValid());
 	}
 
 
diff --git a/Engine/Plugins/Experimental/DigitalMultiplex/DMXProtocol/Source/DMXProtocol/Public/DMXProtocolCommon.h b/Engine/Plugins/Experimental/DigitalMultiplex/DMXProtocol/Source/DMXProtocol/Public/DMXProtocolCommon.h
--- a/Engine/Plugins/Experimental/DigitalMultiplex/DMXProtocol/Source/DMXProtocol/Public/DMXProtocolCommon.h
+++ b/Engine/Plugins/Experimental/DigitalMultiplex/DMXProtocol/Source/DMXProtocol/Public/DMXProtocolCommon.h
@@ -9,22 +9,14 @@ class FBufferArchive;
 class FJsonObject;
 class FArrayReader;
 class IDMXProtocol;
-class IDMXProtocolDevice;
 class IDMXProtocolFactory;
-class IDMXProtocolInterface;
-class IDMXProtocolInterfaceUSB;
-class IDMXProtocolInterfaceEthernet;
-class IDMXProtocolPort;
 class IDMXProtocolRDM;
 class IDMXProtocolUniverse;
 class IDMXProtocolSender;
 class IDMXProtocolReceiver;
 
-class FDMXProtocolDeviceManager;
-class FDMXProtocolInterfaceManager;
-class FDMXProtocolPortManager;
+template<class TUniverse>
 class FDMXProtocolUniverseManager;
-
 class FDMXProtocolPackager;
 
 struct FDMXBuffer;
@@ -32,15 +24,10 @@ struct FDMXPacket;
 struct IDMXProtocolPacket;
 
 using FDataPtr = TSharedPtr<TArray<uint8>>;
-using IDevicesMap = TMap<IDMXProtocolInterface*, TSharedPtr<IDMXProtocolDevice>>;
-using IInterfacesMap = TMap<uint32, TSharedPtr<IDMXProtocolInterface>>;
-
-using IPortsMap = TMap<uint8, TSharedPtr<IDMXProtocolPort>>;
-using IPortsDeviceMap = TMap<IDMXProtocolDevice*, IPortsMap>;
-using IUniversesMap = TMap<IDMXProtocolPort*, TSharedPtr<IDMXProtocolUniverse>>;
-using IUniversesIDMap = TMap<uint16, IDMXProtocolPort*>;
-
 using IDMXFragmentMap = TMap<uint32, uint8>;
 
 using FArrayReaderPtr = TSharedPtr<FArrayReader, ESPMode::ThreadSafe>;
-using FDMXPacketPtr = TSharedPtr<FDMXPacket, ESPMode::ThreadSafe>;
\ No newline at end of file
+using FDMXPacketPtr = TSharedPtr<FDMXPacket, ESPMode::ThreadSafe>;
+using IDMXProtocolUniversePtr = TSharedPtr<IDMXProtocolUniverse, ESPMode::ThreadSafe>;
+
+enum class EDMXSendResult : uint8;
diff --git a/Engine/Plugins/Experimental/DigitalMultiplex/DMXProtocol/Source/DMXProtocol/Public/DMXProtocolConstants.h b/Engine/Plugins/Experimental/DigitalMultiplex/DMXProtocol/Source/DMXProtocol/Public/DMXProtocolConstants.h
--- a/Engine/Plugins/Experimental/DigitalMultiplex/DMXProtocol/Source/DMXProtocol/Public/DMXProtocolConstants.h
+++ b/Engine/Plugins/Experimental/DigitalMultiplex/DMXProtocol/Source/DMXProtocol/Public/DMXProtocolConstants.h
@@ -18,9 +18,19 @@ enum EDMXPortCapability
 	DMX_PORT_CAPABILITY_FULL,
 };
 
-enum EDMXPortDirection
+enum EDMUniverseDirection
 {
-	DMX_PORT_UNKNOWN,
-	DMX_PORT_OUTPUT,
-	DMX_PORT_INPUT
-};
\ No newline at end of file
+	DMX_UNIVERSE_UNKNOWN,
+	DMX_UNIVERSE_OUTPUT,
+	DMX_UNIVERSE_INPUT
+};
+
+#define DMX_MAX_NETPIN 137
+#define DMX_MAX_SUBNETPIN 15
+#define DMX_MAX_UNIVERSE 63999
+#define DMX_MAX_UNIVERSEARTNET 15
+#define DMX_MAX_FINALUNIVERSE 32768
+#define DMX_MAX_ADDRESS 512
+#define DMX_MAX_VALUE 255
+
+#define DMX_K2_CATEGORY_NAME "DMX"
\ No newline at end of file
diff --git a/Engine/Plugins/Experimental/DigitalMultiplex/DMXProtocol/Source/DMXProtocol/Public/DMXProtocolMacros.h b/Engine/Plugins/Experimental/DigitalMultiplex/DMXProtocol/Source/DMXProtocol/Public/DMXProtocolMacros.h
--- a/Engine/Plugins/Experimental/DigitalMultiplex/DMXProtocol/Source/DMXProtocol/Public/DMXProtocolMacros.h
+++ b/Engine/Plugins/Experimental/DigitalMultiplex/DMXProtocol/Source/DMXProtocol/Public/DMXProtocolMacros.h
@@ -3,6 +3,25 @@
 #pragma once
 
 #include "CoreMinimal.h"
+#include "Stats/Stats.h"
+
+DMXPROTOCOL_API DECLARE_LOG_CATEGORY_EXTERN(LogDMXProtocol, Log, All);
+
+DECLARE_STATS_GROUP(TEXT("DMXProtocol"), STATGROUP_DMXProtocol, STATCAT_Advanced);
+
+#ifndef DMXPROTOCOL_LOG_PREFIX
+#define DMXPROTOCOL_LOG_PREFIX TEXT("DMX: ")
+#endif
+
+#define UE_LOG_DMXPROTOCOL(Verbosity, Format, ...) \
+{ \
+	UE_LOG(LogDMXProtocol, Verbosity, TEXT("%s%s"), DMXPROTOCOL_LOG_PREFIX, *FString::Printf(Format, ##__VA_ARGS__)); \
+}
+
+#define UE_CLOG_DMXPROTOCOL(Conditional, Verbosity, Format, ...) \
+{ \
+	UE_CLOG(Conditional, LogDMXProtocol, Verbosity, TEXT("%s%s"), DMXPROTOCOL_LOG_PREFIX, *FString::Printf(Format, ##__VA_ARGS__)); \
+}
 
 #define REGISTER_DMX_ARCHIVE(ProtocolDMXPacket) \
 	FArchive & operator<<(FArchive & Ar, ProtocolDMXPacket&Packet) \
diff --git a/Engine/Plugins/Experimental/DigitalMultiplex/DMXProtocol/Source/DMXProtocol/Public/DMXProtocolModule.h b/Engine/Plugins/Experimental/DigitalMultiplex/DMXProtocol/Source/DMXProtocol/Public/DMXProtocolModule.h
--- a/Engine/Plugins/Experimental/DigitalMultiplex/DMXProtocol/Source/DMXProtocol/Public/DMXProtocolModule.h
+++ b/Engine/Plugins/Experimental/DigitalMultiplex/DMXProtocol/Source/DMXProtocol/Public/DMXProtocolModule.h
@@ -5,9 +5,7 @@
 #include "DMXProtocolCommon.h"
 #include "Modules/ModuleInterface.h"
 
-/**
- * 
- */
+/**   */
 class DMXPROTOCOL_API FDMXProtocolModule 
 	: public IModuleInterface
 {
@@ -21,11 +19,16 @@ public:
 	 * @param  InProtocolName Name of the requested protocol
 	 * @return Return the pointer to protocol.
 	 */
-	virtual IDMXProtocol* GetProtocol(const FName InProtocolName = NAME_None);
+	virtual TSharedPtr<IDMXProtocol> GetProtocol(const FName InProtocolName = NAME_None);
 	
+	/**  Get the reference to all protocol factories map */
 	const TMap<FName, IDMXProtocolFactory*>& GetProtocolFactories() const;
 
+	/**  Get the reference to all protocols map */
+	const TMap<FName, TSharedPtr<IDMXProtocol>>& GetProtocols() const;
+
 	//~ Begin IModuleInterface implementation
+	virtual void StartupModule() override;
 	virtual void ShutdownModule() override;
 	//~ End IModuleInterface implementation
 
diff --git a/Engine/Plugins/Experimental/DigitalMultiplex/DMXProtocol/Source/DMXProtocol/Public/DMXProtocolSettings.h b/Engine/Plugins/Experimental/DigitalMultiplex/DMXProtocol/Source/DMXProtocol/Public/DMXProtocolSettings.h
--- a/Engine/Plugins/Experimental/DigitalMultiplex/DMXProtocol/Source/DMXProtocol/Public/DMXProtocolSettings.h
+++ b/Engine/Plugins/Experimental/DigitalMultiplex/DMXProtocol/Source/DMXProtocol/Public/DMXProtocolSettings.h
@@ -4,25 +4,53 @@
 
 #include "CoreMinimal.h"
 #include "UObject/Object.h"
+#include "DMXProtocolTypes.h"
 #include "DMXProtocolSettings.generated.h"
 
-/**
- * User defined protocol settings that apply to a whole protocol module
- */
-UCLASS(config = Engine)
+/**  User defined protocol settings that apply to a whole protocol module */
+UCLASS(config = Engine, notplaceable)
 class DMXPROTOCOL_API UDMXProtocolSettings : public UObject
 {
 public:
 	GENERATED_BODY()
 
 public:
-	UDMXProtocolSettings() :
-		InterfaceIPAddress(TEXT("0.0.0.0"))
-	{}
+	UDMXProtocolSettings();
 
+#if WITH_EDITOR
+	void PostEditChangeProperty(FPropertyChangedEvent& PropertyChangedEvent);
+#endif
 
 public:
 	/** Manual Interface IP Address */
-	UPROPERTY(GlobalConfig, EditAnywhere, Category = "DMX", Meta = (DisplayName = "Interface IP address"))
+	UPROPERTY(Config, EditAnywhere, Category = "DMX|Communication Settings", Meta = (DisplayName = "Interface IP address"))
 	FString InterfaceIPAddress;
+
+	UPROPERTY(Config, EditAnywhere, Category = "DMX|Unicast")
+	bool bShouldUseUnicast;
+
+	UPROPERTY(Config, EditAnywhere, Category = "DMX|Unicast", Meta = (EditCondition="bShouldUseUnicast"))
+	FString UnicastEndpoint;
+
+	/** Universe Remote Start for ArtNet */
+	UPROPERTY(Config, EditAnywhere, Category = "DMX|Universe Settings", Meta = (DisplayName = "Global Art-Net Universe Offset", ClampMin = "0", ClampMax = "65535"))
+	int32 GlobalArtNetUniverseOffset;
+
+	/** Universe Remote Start for sACN */
+	UPROPERTY(Config, EditAnywhere, Category = "DMX|Universe Settings", Meta = (DisplayName = "Global sACN Universe Offset", ClampMin = "0", ClampMax = "65535"))
+	int32 GlobalSACNUniverseOffset;
+
+	/** Fixture Categories ENum */
+	UPROPERTY(Config, EditAnywhere, Category = "DMX|Fixture Settings", Meta = (DisplayName = "Fixture Categories"))
+	TSet<FName> FixtureCategories;
+
+	//~ Properties controlled by the Input Console (SDMXInputInfoSelecter)
+
+	/** Set the current protocol to be monitored */
+	UPROPERTY(Config)
+	FName InputConsoleProtocol;
+
+	/** Set the current Universe ID to be monitored */
+	UPROPERTY(Config)
+	uint16 InputConsoleUniverseID;
 };
diff --git a/Engine/Plugins/Experimental/DigitalMultiplex/DMXProtocol/Source/DMXProtocol/Public/DMXProtocolTypes.h b/Engine/Plugins/Experimental/DigitalMultiplex/DMXProtocol/Source/DMXProtocol/Public/DMXProtocolTypes.h
--- a/Engine/Plugins/Experimental/DigitalMultiplex/DMXProtocol/Source/DMXProtocol/Public/DMXProtocolTypes.h
+++ b/Engine/Plugins/Experimental/DigitalMultiplex/DMXProtocol/Source/DMXProtocol/Public/DMXProtocolTypes.h
@@ -2,24 +2,206 @@
 
 #pragma once
 
-#include "Interfaces/IDMXProtocol.h"
 #include "DMXProtocolCommon.h"
 #include "DMXProtocolConstants.h"
+#include "DMXProtocolMacros.h"
 #include "Dom/JsonObject.h"
 
+#include "Kismet/BlueprintFunctionLibrary.h"
+
+#include "DMXProtocolTypes.generated.h"
+
+UENUM(BlueprintType)
+enum class EDMXSendResult : uint8
+{
+	Success UMETA(DisplayName = "Successfully sent"),
+	ErrorGetUniverse UMETA(DisplayName = "Error Get Universe"),
+	ErrorSetBuffer UMETA(DisplayName = "Error Set Buffer"),
+	ErrorSizeBuffer UMETA(DisplayName = "Error Size Buffer"),
+	ErrorEnqueuePackage UMETA(DisplayName = "Error Enqueue Package"),
+	ErrorNoSenderInterface UMETA(DisplayName = "Error No Sending Interface")
+};
+
+UENUM()
+enum class EDMXProtocolDirectionality : uint8
+{
+	EInput 	UMETA(DisplayName = "Input"),
+	EOutput UMETA(DisplayName = "Output"),
+};
+
+UENUM()
+enum class EDMXCommunicationTypes : uint8
+{
+	Broadcast UMETA(DisplayName = "Broadcast")
+};
+
+UENUM()
+enum class EDMXFixtureSignalFormat : uint8
+{
+	/** Uses 1 channel (byte) and allows subdivision into sub functions */
+	E8BitSubFunctions UMETA(DisplayName = "8 Bit (Sub Functions)", Hidden),
+	/** Uses 1 channel (byte). Range: 0 to 255 */
+	E8Bit 	UMETA(DisplayName = "8 Bit"),
+	/** Uses 2 channels (bytes). Range: 0 to 65.535 */
+	E16Bit 	UMETA(DisplayName = "16 Bit"),
+	/** Uses 3 channels (bytes). Range: 0 to 16.777.215 */
+	E24Bit 	UMETA(DisplayName = "24 Bit"),
+	/** Uses 4 channels (bytes). Range: 0 to 4.294.967.295 */
+	E32Bit 	UMETA(DisplayName = "32 Bit"),
+};
+
+USTRUCT(BlueprintType)
+struct DMXPROTOCOL_API FDMXProtocolName
+{
+public:
+	GENERATED_BODY()
+
+	/** Selected protocol name */
+	UPROPERTY(BlueprintReadWrite, Category = "DMX|Protocol")
+	FName Name;
+
+	static TArray<FName> GetPossibleValues();
+
+	FDMXProtocolName() = default;
+	/** Construct from a protocol name */
+	explicit FDMXProtocolName(const FName& InName);
+	/** Construct from a protocol */
+	FDMXProtocolName(TSharedPtr<IDMXProtocol> InProtocol);
+
+	/** Returns the Protocol this name represents */
+	TSharedPtr<IDMXProtocol> GetProtocol() const;
+	/** True if it has a valid Protocol name */
+	bool IsValid() const;
+
+	//~ FName operators
+	operator FName&() { return Name; }
+	operator const FName&() const { return Name; }
+
+	//~ DMXProtocol operators
+	operator TSharedPtr<IDMXProtocol>() { return GetProtocol(); }
+	operator const TSharedPtr<IDMXProtocol>() const { return GetProtocol(); }
+
+	/** Bool (is valid) operator */
+	operator bool() const { return IsValid(); }
+
+	//~ Comparison operators
+	bool operator==(const FDMXProtocolName& Other) const { return Name == Other.Name; }
+	bool operator==(const TSharedPtr<IDMXProtocol>& Other) const { return GetProtocol().Get() == Other.Get(); }
+	bool operator==(const FName& Other) const { return Name == Other; }
+};
+
+USTRUCT(BlueprintType)
+struct DMXPROTOCOL_API FDMXFixtureCategory
+{
+public:
+	GENERATED_BODY()
+
+	/** Selected protocol name */
+	UPROPERTY(BlueprintReadWrite, Category = "DMX|Category")
+	FName Name;
+
+	static FSimpleMulticastDelegate OnPossibleValuesUpdated;
+	static TArray<FName> GetPossibleValues();
+
+	static FName GetFirstValue();
+
+	FDMXFixtureCategory();
+
+	explicit FDMXFixtureCategory(const FName& InName);
+
+	//~ FName operators
+	operator FName&() { return Name; }
+	operator const FName&() const { return Name; }
+
+	/** Bool (is valid) operator */
+	operator bool() const { return !Name.IsNone(); }
+
+	//~ Comparison operators
+	bool operator==(const FDMXFixtureCategory& Other) const { return Name == Other.Name; }
+	bool operator==(const FName& Other) const { return Name == Other; }
+};
+
+UCLASS()
+class UDMXNameContainersConversions
+	: public UBlueprintFunctionLibrary
+{
+	GENERATED_BODY()
+
+public:
+	UFUNCTION(BlueprintPure, meta = (DisplayName = "ToString (DMX Protocol Name)", CompactNodeTitle = "->", BlueprintAutocast), Category = "Utilities|DMX")
+	static FString Conv_DMXProtocolNameToString(const FDMXProtocolName& InProtocolName);
+
+	UFUNCTION(BlueprintPure, meta = (DisplayName = "ToName (DMX Protocol Name)", CompactNodeTitle = "->", BlueprintAutocast), Category = "Utilities|DMX")
+	static FName Conv_DMXProtocolNameToName(const FDMXProtocolName& InProtocolName);
+
+	UFUNCTION(BlueprintPure, meta = (DisplayName = "ToString (DMX Fixture Category)", CompactNodeTitle = "->", BlueprintAutocast), Category = "Utilities|DMX")
+	static FString Conv_DMXFixtureCategoryToString(const FDMXFixtureCategory& InFixtureCategory);
+
+	UFUNCTION(BlueprintPure, meta = (DisplayName = "ToName (DMX Fixture Category)", CompactNodeTitle = "->", BlueprintAutocast), Category = "Utilities|DMX")
+	static FName Conv_DMXFixtureCategoryToName(const FDMXFixtureCategory& InFixtureCategory);
+};
+
+
+USTRUCT()
+struct DMXPROTOCOL_API FDMXUniverse
+{
+	GENERATED_BODY()
+
+	UPROPERTY(EditAnywhere, Category = "DMX")
+	uint32 UniverseNumber;
+
+	UPROPERTY(EditAnywhere, meta = (DisplayName = "Address", ClampMin = 0, ClampMax = 512, UIMin = 0, UIMax = 512), Category = "DMX")
+	uint32 Channel;
+
+	UPROPERTY()
+	EDMXProtocolDirectionality DMXProtocolDirectionality;
+
+	FDMXUniverse()
+		: UniverseNumber(0)
+		, Channel(1)
+		, DMXProtocolDirectionality(EDMXProtocolDirectionality::EInput)
+	{}
+};
+
 struct DMXPROTOCOL_API FDMXBuffer
 {
 public:
 	FDMXBuffer()
+		: SequanceID(0)
 	{
 		DMXData.AddZeroed(DMX_UNIVERSE_SIZE);
 	}
 
-	bool SetDMXFragment(const IDMXFragmentMap & DMXFragment);
+	/** @return Gets actual SequanceID  */
+	uint32 GetSequanceID() const { return SequanceID; }
+
+	/** @return editable DMX buffer */
 	TArray<uint8>& GetDMXData() { return DMXData; }
 
+	/**
+	 * Updates the fragment in the DMX buffer
+	 * This is set Map values, by specifying channels into the 512 DMX buffer
+	 * 
+	 * @param IDMXFragmentMap DMX fragment map
+	 * @return True if it was successfully set
+	 */
+	bool SetDMXFragment(const IDMXFragmentMap& InDMXFragment);
+
+	/**
+	 * Sets the DMX buffer from input buffer
+	 * 
+	 * @param InBuffer Pointer to DMX buffer array
+	 * @param InSize Size of the buffer to copy 
+	 * @return True if it was successfully set
+	 */
+	bool SetDMXBuffer(const uint8* InBuffer, uint32 InSize);
+
 private:
+	/** DMX bytes buffer array */
 	TArray<uint8> DMXData;
+
+	/** Synchronizations sequence id */
+	uint32 SequanceID;
 };
 
 struct DMXPROTOCOL_API FDMXPacket
@@ -59,7 +241,7 @@ struct DMXPROTOCOL_API FRDMUID
 		}
 		else
 		{
-			UE_LOG_DMXPROTOCOL(Warning, TEXT("Size of the TArray buffer is wrong"));
+			UE_LOG_DMXPROTOCOL(Verbose, TEXT("Size of the TArray buffer is wrong"));
 		}
 	}
 
diff --git a/Engine/Plugins/Experimental/DigitalMultiplex/DMXProtocol/Source/DMXProtocol/Public/Interfaces/IDMXNetworkInterface.h b/Engine/Plugins/Experimental/DigitalMultiplex/DMXProtocol/Source/DMXProtocol/Public/Interfaces/IDMXNetworkInterface.h
new file mode 100644
--- /dev/null
+++ b/Engine/Plugins/Experimental/DigitalMultiplex/DMXProtocol/Source/DMXProtocol/Public/Interfaces/IDMXNetworkInterface.h
@@ -0,0 +1,30 @@
+// Copyright Epic Games, Inc. All Rights Reserved.
+
+#pragma once
+
+#include "CoreMinimal.h"
+#include "DMXProtocolCommon.h"
+
+/**
+ */
+class IDMXNetworkInterface
+{
+public:
+	virtual ~IDMXNetworkInterface() {}
+
+	/**
+	 * Listen for interface IP changes
+	 * @param InInterfaceIPAddress			String IP address
+	 */
+	virtual void OnNetworkInterfaceChanged(const FString& InInterfaceIPAddress) = 0;
+
+	/**
+	 * Restart sockets and listeners
+	 * @param InInterfaceIPAddress			String IP address
+	 * @param OutErrorMessage				String error message
+	 */
+	virtual bool RestartNetworkInterface(const FString& InInterfaceIPAddress, FString& OutErrorMessage) = 0;
+
+	/**  Release network interfaces */
+	virtual void ReleaseNetworkInterface() = 0;
+};
\ No newline at end of file
diff --git a/Engine/Plugins/Experimental/DigitalMultiplex/DMXProtocol/Source/DMXProtocol/Public/Interfaces/IDMXProtocol.h b/Engine/Plugins/Experimental/DigitalMultiplex/DMXProtocol/Source/DMXProtocol/Public/Interfaces/IDMXProtocol.h
--- a/Engine/Plugins/Experimental/DigitalMultiplex/DMXProtocol/Source/DMXProtocol/Public/Interfaces/IDMXProtocol.h
+++ b/Engine/Plugins/Experimental/DigitalMultiplex/DMXProtocol/Source/DMXProtocol/Public/Interfaces/IDMXProtocol.h
@@ -11,57 +11,57 @@
 #include "Interfaces/IDMXProtocolRDM.h"
 #include "Interfaces/IDMXProtocolTransport.h"
 
-class FDMXProtocolDeviceManager;
-class FDMXProtocolInterfaceManager;
-class FDMXProtocolPortManager;
-class FDMXProtocolUniverseManager;
-
-DMXPROTOCOL_API DECLARE_LOG_CATEGORY_EXTERN(LogDMXProtocol, Log, All);
-
-DECLARE_STATS_GROUP(TEXT("DMXProtocol"), STATGROUP_DMXProtocol, STATCAT_Advanced);
-
-#ifndef DMXPROTOCOL_LOG_PREFIX
-#define DMXPROTOCOL_LOG_PREFIX TEXT("DMX: ")
-#endif
-
-#define UE_LOG_DMXPROTOCOL(Verbosity, Format, ...) \
-{ \
-	UE_LOG(LogDMXProtocol, Verbosity, TEXT("%s%s"), DMXPROTOCOL_LOG_PREFIX, *FString::Printf(Format, ##__VA_ARGS__)); \
-}
-
-#define UE_CLOG_DMXPROTOCOL(Conditional, Verbosity, Format, ...) \
-{ \
-	UE_CLOG(Conditional, LogDMXProtocol, Verbosity, TEXT("%s%s"), DMXPROTOCOL_LOG_PREFIX, *FString::Printf(Format, ##__VA_ARGS__)); \
-}
-
 /**
- * Generic protocol interface, it should be inherited by all protocol implementations.
+ * Delegate used when downloading of message contents has completed
+ *
+ * @param LocalUserNum the controller number of the associated user that made the request
+ * @param bWasSuccessful true if the async action completed without error, false if there was an error
+ * @param MessageId unique id of the message downloaded
+ * @param ErrorStr string representing the error condition
  */
+DECLARE_MULTICAST_DELEGATE_OneParam(FOnNetworkInterfaceChanged, const FString&);
+typedef FOnNetworkInterfaceChanged::FDelegate FOnNetworkInterfaceChangedDelegate;
+
+struct FDMXUniverse;
+
+/**  Generic protocol interface, it should be inherited by all protocol implementations. */
 class DMXPROTOCOL_API IDMXProtocol 
 	: public IDMXProtocolBase
 	, public IDMXProtocolRDM
 {
 public:
-	virtual ~IDMXProtocol() {}
-
-	/**
-	 * If protocol exists return the pointer otherwise it create a new protocol first and then return the pointer.
-	 * @param  InProtocolName Name of the requested protocol
-	 * @return Return the pointer to protocol
-	 */
-	template <class TProtocol>
-	static TProtocol* Get(const FName& ProtocolName = NAME_None)
+	static const TMap<FName, IDMXProtocolFactory*>& GetProtocolFactories()
 	{
 		static const FName DMXProtocolModuleName = TEXT("DMXProtocol");
 		FDMXProtocolModule& DMXProtocolModule = FModuleManager::GetModuleChecked<FDMXProtocolModule>(DMXProtocolModuleName);
-		return static_cast<TProtocol*>(DMXProtocolModule.GetProtocol(ProtocolName));
+		return DMXProtocolModule.GetProtocolFactories();
 	}
 
-	static const TMap<FName, IDMXProtocolFactory*>& GetProtocolFactories()
+	static const TMap<FName, TSharedPtr<IDMXProtocol>>& GetProtocols()
 	{
 		static const FName DMXProtocolModuleName = TEXT("DMXProtocol");
 		FDMXProtocolModule& DMXProtocolModule = FModuleManager::GetModuleChecked<FDMXProtocolModule>(DMXProtocolModuleName);
-		return DMXProtocolModule.GetProtocolFactories();
+		return DMXProtocolModule.GetProtocols();
+	}
+
+	static const TArray<FName>& GetProtocolNames()
+	{
+		const TMap<FName, IDMXProtocolFactory*>& Protocols = GetProtocolFactories();
+		static TArray<FName> ProtocolNames;
+		Protocols.GenerateKeyArray(ProtocolNames);
+		return ProtocolNames;
+	}
+
+	static FName GetFirstProtocolName()
+	{
+		const TMap<FName, IDMXProtocolFactory*>& ProtocolFactories = IDMXProtocol::GetProtocolFactories();
+
+		for (const auto& Itt : ProtocolFactories)
+		{
+			return Itt.Key;
+		}
+
+		return FName();
 	}
 
 	/**
@@ -69,7 +69,7 @@ public:
 	 * @param  InProtocolName Name of the requested protocol
 	 * @return Return the pointer to protocol
 	 */
-	static IDMXProtocol* Get(const FName& ProtocolName = NAME_None)
+	static TSharedPtr<IDMXProtocol> Get(const FName& ProtocolName = NAME_None)
 	{
 		static const FName DMXProtocolModuleName = TEXT("DMXProtocol");
 		FDMXProtocolModule& DMXProtocolModule = FModuleManager::GetModuleChecked<FDMXProtocolModule>(DMXProtocolModuleName);
@@ -80,7 +80,7 @@ public:
 	 * Get the Protocol Name
 	 * @return Return FName of the protocol
 	 */
-	virtual FName GetProtocolName() const = 0;
+	virtual const FName& GetProtocolName() const = 0;
 	
 	/**
 	 * Get the Protocol Sender Interface
@@ -90,67 +90,101 @@ public:
 	virtual TSharedPtr<IDMXProtocolSender> GetSenderInterface() const = 0;
 
 	/**
-	 * Get the device manager for a protocol instance
-	 * The device manager is responsible for hold and searches the physical devices
-	 * Physical devices are nodes, controllers, consoles
-	 * @return Return the pointer to DeviceManager
+	 * Get the protocol settings
+	 * @return Return the pointer to Protocol Settings
 	 */
-	virtual TSharedPtr<FDMXProtocolDeviceManager> GetDeviceManager() const = 0;
-	
+	virtual TSharedPtr<FJsonObject> GetSettings() const = 0;
+
 	/**
-	 * Get the interface manager for a protocol instance
-	 * The interface is the way how protocol interacts with Unreal Engine
-	 * This could be done with ethernet, serial or fake communication
-	 * @return Return the pointer to InterfaceManager
+	 * Whether protocol enabled
+	 * @return Return true if the protocol is enabled
 	 */
-	virtual TSharedPtr<FDMXProtocolInterfaceManager> GetInterfaceManager() const = 0;
+	virtual bool IsEnabled() const = 0;
 
 	/**
-	 * Get the port manager for a protocol instance
-	 * Port managers are responsible for hold the configuration instances of physical ports
+	 * Add universe to the manager
+	 * @param  FJsonObject universe settings, such as UniverseID, Subnet, etc.
+	 * This is unique to each protocol implementation
+	 * @return Return the pointer to universe
 	 */
-	virtual TSharedPtr<FDMXProtocolPortManager> GetPortManager() const = 0;
-	
+	virtual TSharedPtr<IDMXProtocolUniverse, ESPMode::ThreadSafe> AddUniverse(const FJsonObject& InSettings) = 0;
+
 	/**
-	 * Get the universe manager for a protocol instance
-	 * Universe managers hold the buffers of DMX for particular universes
-	 * @return Return the pointer to UniverseManager
+	 * Collects the universes related to a UniverseManger Entity and add them to
+	 * the protocol to be used for communication.
+	 * @param Universes The list of universes from the Entity.
 	 */
-	virtual TSharedPtr<FDMXProtocolUniverseManager> GetUniverseManager() const = 0;
+	virtual void CollectUniverses(const TArray<FDMXUniverse>& Universes) = 0;
 
 	/**
-	 * Get the protocol settings
-	 * @return Return the pointer to Protocol Settings
+	 * Remove Universe from the Protocol Universe Manager.
+	 * @param  InUniverseId unique number of universe
+	 * @return Return true if it was successfully removed
 	 */
-	virtual TSharedPtr<FJsonObject> GetSettings() const = 0;
+	virtual bool RemoveUniverseById(uint32 InUniverseId) = 0;
+
+	/**  Remove all universes from protocol manager */
+	virtual void RemoveAllUniverses() = 0;
 
 	/**
-	 * Reloads the protocol instance
+	 * Getting Universe from the Protocol Universe Manager.
+	 * @param  InUniverseId unique number of universe
+	 * @return Return the pointer to the universe, or nullptr if there is no Universe by the given ID
 	 */
-	virtual void Reload() = 0;
+	virtual TSharedPtr<IDMXProtocolUniverse, ESPMode::ThreadSafe> GetUniverseById(uint32 InUniverseId) const = 0;
 
 	/**
-	 * Whether protocol enabled
-	 * @return Return true if the protocol is enabled
+	 * Get current amount of universes in the Map
+	 * @return Return amount of universes in the Map
 	 */
-	virtual bool IsEnabled() const = 0;
+	virtual uint32 GetUniversesNum() const = 0;
+
+	/**
+	 * Get minimum supported universe ID for protocol
+	 * @return Minimum supported universe ID for protocol
+	 */
+	virtual uint16 GetMinUniverseID() const = 0;
+	
+	/**
+	 * Get maximum supported universes in protocol
+	 * @return Maximum supported universes in protocol
+	 */
+	virtual uint16 GetMaxUniverses() const = 0;
 
 	/**
 	 * Sets the DMX fragment for a particular universe
 	 * @param  UniverseID ID of universe to send
 	 * @param  DMXFragment Map of DMX channel  and values
-	 * @param  bShouldSend whether it should be sent or just store in the universe buffer
+	 * @return Return the status of sending
 	 */
-	virtual void SetDMXFragment(uint16 UniverseID, const IDMXFragmentMap& DMXFragment, bool bShouldSend = true) = 0;
+	virtual EDMXSendResult SendDMXFragment(uint16 UniverseID, const IDMXFragmentMap& DMXFragment) = 0;
 
 	/**
-	 * Send DMX to physical device port
+	 * Sets the DMX fragment for a particular universe
+	 * Create protocol Universe if it does not exist
 	 * @param  UniverseID ID of universe to send
-	 * @param  PortID Physical port ID
-	 * @param  DMXBuffer DMX buffer to sent
-	 * @return Return true if DMX buffer included in the send queue
+	 * @param  DMXFragment Map of DMX channel  and values
+	 * @return Return the status of sending
 	 */
-	virtual bool SendDMX(uint16 UniverseID, uint8 PortID, const TSharedPtr<FDMXBuffer>& DMXBuffer) const = 0;
+	virtual EDMXSendResult SendDMXFragmentCreate(uint16 UniverseID, const IDMXFragmentMap& DMXFragment) = 0;
+
+	/**
+	 * Gets the final protocol universe ID to send
+	 * This is implemented protocol-specific offset
+	 * @param  UniverseID ID of universe to send
+	 * @return Return final Universe ID for sending
+	 */
+	virtual uint16 GetFinalSendUniverseID(uint16 InUniverseID) const = 0;
+
+	/**
+	 * Called on input universe.
+	 * Parameters represent: Protocol Name, UniverseID and Buffer
+	 */
+	DECLARE_EVENT_ThreeParams(IDMXProtocol, FOnUniverseInputUpdateEvent, FName, uint16, const TArray<uint8>&);
+	virtual FOnUniverseInputUpdateEvent& GetOnUniverseInputUpdate() = 0;
+
+public:
+	static FOnNetworkInterfaceChanged OnNetworkInterfaceChanged;
 };
 
 
diff --git a/Engine/Plugins/Experimental/DigitalMultiplex/DMXProtocol/Source/DMXProtocol/Public/Interfaces/IDMXProtocolBase.h b/Engine/Plugins/Experimental/DigitalMultiplex/DMXProtocol/Source/DMXProtocol/Public/Interfaces/IDMXProtocolBase.h
--- a/Engine/Plugins/Experimental/DigitalMultiplex/DMXProtocol/Source/DMXProtocol/Public/Interfaces/IDMXProtocolBase.h
+++ b/Engine/Plugins/Experimental/DigitalMultiplex/DMXProtocol/Source/DMXProtocol/Public/Interfaces/IDMXProtocolBase.h
@@ -5,8 +5,11 @@
 #include "CoreMinimal.h"
 
 class DMXPROTOCOL_API IDMXProtocolBase
+	: public TSharedFromThis<IDMXProtocolBase>
 {
 public:
+	virtual ~IDMXProtocolBase() {}
+
 	virtual bool Init() = 0;
 
 	virtual bool Shutdown() = 0;
diff --git a/Engine/Plugins/Experimental/DigitalMultiplex/DMXProtocol/Source/DMXProtocol/Public/Interfaces/IDMXProtocolTransport.h b/Engine/Plugins/Experimental/DigitalMultiplex/DMXProtocol/Source/DMXProtocol/Public/Interfaces/IDMXProtocolTransport.h
--- a/Engine/Plugins/Experimental/DigitalMultiplex/DMXProtocol/Source/DMXProtocol/Public/Interfaces/IDMXProtocolTransport.h
+++ b/Engine/Plugins/Experimental/DigitalMultiplex/DMXProtocol/Source/DMXProtocol/Public/Interfaces/IDMXProtocolTransport.h
@@ -10,9 +10,7 @@
 
 class FRunnableThread;
 
-/**
- * Delegate type for received data.
- */
+/**  Delegate type for received data. */
 DECLARE_DELEGATE_OneParam(FOnDMXDataReceived, const FArrayReaderPtr&);
 
 class IDMXProtocolSender 
diff --git a/Engine/Plugins/Experimental/DigitalMultiplex/DMXProtocol/Source/DMXProtocol/Public/Interfaces/IDMXProtocolUniverse.h b/Engine/Plugins/Experimental/DigitalMultiplex/DMXProtocol/Source/DMXProtocol/Public/Interfaces/IDMXProtocolUniverse.h
--- a/Engine/Plugins/Experimental/DigitalMultiplex/DMXProtocol/Source/DMXProtocol/Public/Interfaces/IDMXProtocolUniverse.h
+++ b/Engine/Plugins/Experimental/DigitalMultiplex/DMXProtocol/Source/DMXProtocol/Public/Interfaces/IDMXProtocolUniverse.h
@@ -4,16 +4,17 @@
 
 #include "DMXProtocolCommon.h"
 
-class DMXPROTOCOL_API IDMXProtocolUniverse
+class IDMXProtocolUniverse
 {
 public:
 	virtual ~IDMXProtocolUniverse() {}
 
-	virtual IDMXProtocol* GetProtocol() const = 0;
-	virtual TSharedPtr<FDMXBuffer> GetOutputDMXBuffer() const = 0;
+	virtual TSharedPtr<IDMXProtocol> GetProtocol() const = 0;
 	virtual TSharedPtr<FDMXBuffer> GetInputDMXBuffer() const = 0;
+	virtual TSharedPtr<FDMXBuffer> GetOutputDMXBuffer() const = 0;
 	virtual bool SetDMXFragment(const IDMXFragmentMap& DMXFragment) = 0;
 	virtual uint8 GetPriority() const = 0;
-	virtual uint16 GetUniverseID() const = 0;
-	virtual TWeakPtr<IDMXProtocolPort> GetCachedUniversePort() const = 0;
+	virtual uint32 GetUniverseID() const = 0;
+	virtual TSharedPtr<FJsonObject> GetSettings() const = 0;
+	virtual bool IsSupportRDM() const = 0;
 };
