- DMX Editor

diff --git a/Engine/Plugins/Experimental/DigitalMultiplex/DMXProtocol/Source/DMXProtocol/Public/Managers/DMXProtocolUniverseManager.h b/Engine/Plugins/Experimental/DigitalMultiplex/DMXProtocol/Source/DMXProtocol/Public/Managers/DMXProtocolUniverseManager.h
--- a/Engine/Plugins/Experimental/DigitalMultiplex/DMXProtocol/Source/DMXProtocol/Public/Managers/DMXProtocolUniverseManager.h
+++ b/Engine/Plugins/Experimental/DigitalMultiplex/DMXProtocol/Source/DMXProtocol/Public/Managers/DMXProtocolUniverseManager.h
@@ -5,22 +5,88 @@
 #include "DMXProtocolCommon.h"
 #include "Interfaces/IDMXProtocol.h"
 
-class DMXPROTOCOL_API FDMXProtocolUniverseManager
+#include "Dom/JsonObject.h"
+
+template<class TUniverse>
+class FDMXProtocolUniverseManager
 {
+public:
+	using TUniversesPtr = TSharedPtr<TUniverse, ESPMode::ThreadSafe>;
+	using IUniversesMap = TMap<uint32, TUniversesPtr>;
+
 public:
 	FDMXProtocolUniverseManager(IDMXProtocol* InDMXProtocol)
 		: DMXProtocol(InDMXProtocol)
 	{}
 
-public:
-	void AddUniverse(TSharedPtr<IDMXProtocolPort> InPort, TSharedPtr<IDMXProtocolUniverse>  InUniverse);
-	void RemoveUniverse(TSharedPtr<IDMXProtocolPort> InPort, TSharedPtr<IDMXProtocolUniverse>  InUniverse);
-	void RemoveAll();
-	IDMXProtocolUniverse* GetUniverseByPort(IDMXProtocolPort* InPort);
-	IDMXProtocolUniverse* GetUniverseById(uint16 UniverseID);
+	TUniversesPtr AddUniverse(uint32 InUniverseID, TUniversesPtr InUniverse)
+	{
+		return UniversesMap.Add(InUniverseID, InUniverse);
+	}
+
+	TUniversesPtr AddUniverseCreate(uint32 InUniverseID)
+	{
+		TUniversesPtr* UniverseAddr = UniversesMap.Find(InUniverseID);
+
+		// Try to find universe
+		if (UniverseAddr != nullptr)
+		{
+			return *UniverseAddr;
+		}
+		// Otherwise create new and return new instance 
+		else
+		{
+			FJsonObject UniverseSettings;
+			UniverseSettings.SetNumberField(TEXT("UniverseID"), InUniverseID);
+			UniverseSettings.SetNumberField(TEXT("PortID"), 0); // TODO set PortID
+
+			TUniversesPtr Universe = StaticCastSharedPtr<TUniverse>(DMXProtocol->AddUniverse(UniverseSettings));
+
+			UniversesMap.Add(InUniverseID, Universe);
+
+			return Universe;
+		}
+	}
+
+	bool RemoveUniverseById(uint32 UniverseID)
+	{
+		TUniversesPtr Universe = GetUniverseById(UniverseID);
+		if (Universe.IsValid())
+		{
+			UniversesMap.Remove(UniverseID);
+			return true;
+		}
+		else
+		{
+			UE_LOG_DMXPROTOCOL(Verbose, TEXT("UniverseID %d doesn't exists"), UniverseID);
+			return false;
+		}
+	}
+
+	void RemoveAll()
+	{
+		UniversesMap.Empty();
+	}
+
+	TUniversesPtr GetUniverseById(uint32 UniverseID)
+	{
+		TUniversesPtr* UniverseAddr = UniversesMap.Find(UniverseID);
+
+		if (UniverseAddr != nullptr)
+		{
+			return *UniverseAddr;
+		}
+
+		return nullptr;
+	}
+
+	const TMap<uint32, TUniversesPtr>& GetAllUniverses() const
+	{
+		return UniversesMap; 
+	}
 
 private:
 	IUniversesMap UniversesMap;
-	IUniversesIDMap UniversesIDMap;
+
 	IDMXProtocol* DMXProtocol;
 };
diff --git a/Engine/Plugins/Experimental/DigitalMultiplex/DMXProtocol/Source/DMXProtocolArtNet/Private/DMXProtocolArtNet.cpp b/Engine/Plugins/Experimental/DigitalMultiplex/DMXProtocol/Source/DMXProtocolArtNet/Private/DMXProtocolArtNet.cpp
--- a/Engine/Plugins/Experimental/DigitalMultiplex/DMXProtocol/Source/DMXProtocolArtNet/Private/DMXProtocolArtNet.cpp
+++ b/Engine/Plugins/Experimental/DigitalMultiplex/DMXProtocol/Source/DMXProtocolArtNet/Private/DMXProtocolArtNet.cpp
@@ -7,14 +7,27 @@
 #include "DMXProtocolSettings.h"
 #include "Managers/DMXProtocolUniverseManager.h"
 #include "Interfaces/IDMXProtocolUniverse.h"
-#include "DMXProtocolArtNetUtils.h"
-#include "DMXProtocolPortArtNet.h"
 #include "DMXProtocolPackager.h"
-#include "Managers/DMXProtocolDeviceManager.h"
-#include "Managers/DMXProtocolPortManager.h"
+
+#include "DMXProtocolArtNetUtils.h"
+#include "DMXProtocolUniverseArtNet.h"
+
+#include "Managers/DMXProtocolUniverseManager.h"
 
 using namespace ArtNet;
 
+FDMXProtocolArtNet::FDMXProtocolArtNet(const FName& InProtocolName, const FJsonObject& InSettings)
+	: ProtocolName(InProtocolName)
+	, BroadcastSocket(nullptr)
+	, ListeningSocket(nullptr)
+{
+	NetworkErrorMessagePrefix = TEXT("NETWORK ERROR Art-Net:");
+
+	Settings = MakeShared<FJsonObject>(InSettings);
+	UniverseManager = MakeShared<FDMXProtocolUniverseManager<FDMXProtocolUniverseArtNet>>(this);
+}
+
+
 TSharedPtr<IDMXProtocolSender> FDMXProtocolArtNet::GetSenderInterface() const
 {
 	return ArtNetSender;
@@ -22,76 +35,267 @@ TSharedPtr<IDMXProtocolSender> FDMXProtocolArtNet::GetSenderInterface() const
 
 bool FDMXProtocolArtNet::Init()
 {
+	InterfaceIPAddress = GetDefault<UDMXProtocolSettings>()->InterfaceIPAddress;
+
+	// Set Network Interface listener
+	NetworkInterfaceChangedHandle = IDMXProtocol::OnNetworkInterfaceChanged.AddRaw(this, &FDMXProtocolArtNet::OnNetworkInterfaceChanged);
+
+	// Set Network Interface
+	FString ErrorMessage;
+	if (!RestartNetworkInterface(InterfaceIPAddress, ErrorMessage))
+	{
+		UE_LOG_DMXPROTOCOL(Error, TEXT("%s:%s"), NetworkErrorMessagePrefix, *ErrorMessage);
+	}
+
+	return true;
+}
+
+bool FDMXProtocolArtNet::Shutdown()
+{
+	ReleaseNetworkInterface();
+	IDMXProtocol::OnNetworkInterfaceChanged.Remove(NetworkInterfaceChangedHandle);
+
+	return true;
+}
+
+bool FDMXProtocolArtNet::IsEnabled() const
+{
+	return true;
+}
+
+TSharedPtr<IDMXProtocolUniverse, ESPMode::ThreadSafe> FDMXProtocolArtNet::AddUniverse(const FJsonObject& InSettings)
+{
+	TSharedPtr<FDMXProtocolUniverseArtNet, ESPMode::ThreadSafe> Universe = MakeShared<FDMXProtocolUniverseArtNet, ESPMode::ThreadSafe>(SharedThis(this), InSettings);
+	return UniverseManager->AddUniverse(Universe->GetUniverseID(), Universe);
+}
+
+void FDMXProtocolArtNet::CollectUniverses(const TArray<FDMXUniverse>& Universes)
+{
+	for (const FDMXUniverse& Universe : Universes)
+	{
+		if (UniverseManager->GetAllUniverses().Contains(Universe.UniverseNumber))
+		{
+			continue;
+		}
+
+		FJsonObject UniverseSettings;
+		UniverseSettings.SetNumberField(TEXT("UniverseID"), Universe.UniverseNumber);
+		UniverseSettings.SetNumberField(TEXT("PortID"), 0); // TODO get correct PortID
+		AddUniverse(UniverseSettings);
+	}
+}
+
+bool FDMXProtocolArtNet::RemoveUniverseById(uint32 InUniverseId)
+{
+	return UniverseManager->RemoveUniverseById(InUniverseId);
+}
+
+void FDMXProtocolArtNet::RemoveAllUniverses()
+{
+	UniverseManager->RemoveAll();
+}
+
+TSharedPtr<IDMXProtocolUniverse, ESPMode::ThreadSafe> FDMXProtocolArtNet::GetUniverseById(uint32 InUniverseId) const
+{
+	return UniverseManager->GetUniverseById(InUniverseId);
+}
+
+uint32 FDMXProtocolArtNet::GetUniversesNum() const
+{
+	return UniverseManager->GetAllUniverses().Num();
+}
+
+uint16 FDMXProtocolArtNet::GetMinUniverseID() const
+{
+	return 0;
+}
+
+uint16 FDMXProtocolArtNet::GetMaxUniverses() const
+{
+	return ARTNET_MAX_UNIVERSES;
+}
+
+void FDMXProtocolArtNet::OnNetworkInterfaceChanged(const FString& InInterfaceIPAddress)
+{
+	FString ErrorMessage;
+	if (!RestartNetworkInterface(InInterfaceIPAddress, ErrorMessage))
+	{
+		UE_LOG_DMXPROTOCOL(Error, TEXT("%s:%s"), NetworkErrorMessagePrefix, *ErrorMessage);
+	}
+}
+
+bool FDMXProtocolArtNet::RestartNetworkInterface(const FString& InInterfaceIPAddress, FString& OutErrorMessage)
+{
+	FScopeLock Lock(&SocketsCS);
+
+	// Clean the error message
+	OutErrorMessage.Empty();
+
+	// Try to create IP address at the first
 	ISocketSubsystem* SocketSubsystem = ISocketSubsystem::Get(PLATFORM_SOCKETSUBSYSTEM);
+	// Create sender address
+	SenderAddr = SocketSubsystem->CreateInternetAddr();
+	bool bIsValid = false;
+	SenderAddr->SetIp(*InInterfaceIPAddress, bIsValid);
+	SenderAddr->SetPort(ARTNET_SENDER_PORT);
+	SenderEndpoint = FIPv4Endpoint(SenderAddr);
+	if (!bIsValid)
+	{
+		OutErrorMessage = FString::Printf(TEXT("Wrong IP address: %s"), *InInterfaceIPAddress);
+		return false;
+	}
+
+	// Release old network interface
+	ReleaseNetworkInterface();
+	bool bShouldUseUnicast = GetDefault<UDMXProtocolSettings>()->bShouldUseUnicast;
 
-	// Create broadcast address
 	BroadcastAddr = SocketSubsystem->CreateInternetAddr();
-	BroadcastAddr->SetBroadcastAddress();
+	if (bShouldUseUnicast)
+	{
+		bool bIsUnicastIPValid = false;
+		FString UnicastEndpoint = GetDefault<UDMXProtocolSettings>()->UnicastEndpoint;
+		BroadcastAddr->SetIp(*UnicastEndpoint, bIsUnicastIPValid);
+		if (!bIsUnicastIPValid)
+		{
+			OutErrorMessage = FString::Printf(TEXT("Error Invalid Unicast Address: %s"), *UnicastEndpoint);
+		}
+	}
+	else
+	{
+		BroadcastAddr->SetBroadcastAddress();
+	}
+
 	BroadcastAddr->SetPort(ARTNET_PORT);
 	BroadcastEndpoint = FIPv4Endpoint(BroadcastAddr);
 
-	// Create sender address
-	SenderAddr = SocketSubsystem->CreateInternetAddr();
-	bool bIsValid;
-	FString InterfaceIPAddress = GetDefault<UDMXProtocolSettings>()->InterfaceIPAddress;
-	SenderAddr->SetIp(*InterfaceIPAddress, bIsValid);
-	SenderAddr->SetPort(ARTNET_PORT);
-	SenderEndpoint = FIPv4Endpoint(SenderAddr);
+    FSocket* NewBroadcastSocket = FUdpSocketBuilder(TEXT("UDPArtNetBroadcastSocket"))
+        .AsNonBlocking()
+        .AsReusable()
+        .WithBroadcast()
+        .BoundToEndpoint(SenderEndpoint);
 
-	// Create Broadcast socket
-	BroadcastSocket = FUdpSocketBuilder(TEXT("UDPArtNetBroadcastSocket"))
-		.AsNonBlocking()
-		.AsReusable()
-		.WithBroadcast()
-#if PLATFORM_WINDOWS
-		.BoundToAddress(SenderEndpoint.Address)
-#endif
-		.BoundToPort(SenderEndpoint.Port);
 
-	if (BroadcastSocket == nullptr)
+	if (NewBroadcastSocket == nullptr)
 	{
-		UE_LOG_DMXPROTOCOL(Error, TEXT("ERROR create BroadcastSocket"));
+		OutErrorMessage = FString::Printf(TEXT("Error create BroadcastSocket: %s"), *InterfaceIPAddress);
 		return false;
 	}
 
+	// Set new network interface IP
+	InterfaceIPAddress = InInterfaceIPAddress;
+
+	// Save New socket;
+	BroadcastSocket = NewBroadcastSocket;
+
 	// Create sender
 	ArtNetSender = MakeShared<FDMXProtocolSenderArtNet>(*BroadcastSocket, this);
 
+	// Create Listening socket
+	TSharedPtr<FInternetAddr> ListeningAddr = SocketSubsystem->CreateInternetAddr();
+	ListeningAddr->SetIp(*InInterfaceIPAddress, bIsValid);
+	ListeningAddr->SetPort(ARTNET_PORT);
+	FIPv4Endpoint ListenerEndpoint = FIPv4Endpoint(ListeningAddr);
+
+	FSocket* NewListeningSocket = FUdpSocketBuilder(TEXT("UDPArtNetListeningSocket"))
+		.AsNonBlocking()
+		.AsReusable()
+		.BoundToEndpoint(ListenerEndpoint);
+
+	if (NewListeningSocket == nullptr)
+	{
+		OutErrorMessage = FString::Printf(TEXT("Error create ListeningSocket: %s"), *InterfaceIPAddress);
+		return false;
+	}
+
+	// Save New socket;
+	ListeningSocket = NewListeningSocket;
+
 	// Create receiver
 	FTimespan ThreadWaitTime = FTimespan::FromMilliseconds(100);
-	ArtNetReceiver = MakeShared<FDMXProtocolReceiverArtNet>(*BroadcastSocket, this, ThreadWaitTime);
+	ArtNetReceiver = MakeShared<FDMXProtocolReceiverArtNet>(*ListeningSocket, this, ThreadWaitTime);
 	ArtNetReceiver->OnDataReceived().BindRaw(this, &FDMXProtocolArtNet::OnDataReceived);
 
 	return true;
 }
 
-bool FDMXProtocolArtNet::Shutdown()
+void FDMXProtocolArtNet::ReleaseNetworkInterface()
 {
-	ArtNetSender.Reset();
-	ArtNetReceiver.Reset();
+	if (ArtNetSender.IsValid())
+	{
+		ArtNetSender.Reset();
+		ArtNetSender = nullptr;
+	}
+
+	if (ArtNetReceiver)
+	{
+		ArtNetReceiver.Reset();
+		ArtNetReceiver = nullptr;
+	}
 
-	//Clear all sockets!
-	if (BroadcastSocket)
+	// Clean all sockets!
+	if (BroadcastSocket != nullptr)
 	{
 		BroadcastSocket->Close();
 		ISocketSubsystem::Get(PLATFORM_SOCKETSUBSYSTEM)->DestroySocket(BroadcastSocket);
+		BroadcastSocket = nullptr;
 	}
 
-	return false;
+	if (ListeningSocket != nullptr)
+	{
+		ListeningSocket->Close();
+		ISocketSubsystem::Get(PLATFORM_SOCKETSUBSYSTEM)->DestroySocket(ListeningSocket);
+		ListeningSocket = nullptr;
+	}
 }
 
-bool FDMXProtocolArtNet::IsEnabled() const
+const FName& FDMXProtocolArtNet::GetProtocolName() const
 {
-	return true;
+	return ProtocolName;
 }
 
-FName FDMXProtocolArtNet::GetProtocolName() const
+TSharedPtr<FJsonObject> FDMXProtocolArtNet::GetSettings() const
 {
-	return FDMXProtocolArtNetModule::NAME_Artnet;
+	return Settings;
+}
+
+EDMXSendResult FDMXProtocolArtNet::SendDMXFragment(uint16 InUniverseID, const IDMXFragmentMap& DMXFragment)
+{
+	uint16 FinalSendUniverseID = GetFinalSendUniverseID(InUniverseID);
+
+	TSharedPtr<FDMXProtocolUniverseArtNet, ESPMode::ThreadSafe> Universe = UniverseManager->GetUniverseById(GetFinalSendUniverseID(FinalSendUniverseID));
+	if (!Universe.IsValid())
+	{
+		return EDMXSendResult::ErrorGetUniverse;
+	}
+
+	if (!Universe->SetDMXFragment(DMXFragment))
+	{
+		return EDMXSendResult::ErrorSetBuffer;
+	}
+
+	EDMXSendResult Result = SendDMXInternal(FinalSendUniverseID, Universe->GetPortID(), Universe->GetOutputDMXBuffer());
+	return Result;
 }
 
-void FDMXProtocolArtNet::Reload()
+EDMXSendResult FDMXProtocolArtNet::SendDMXFragmentCreate(uint16 InUniverseID, const IDMXFragmentMap& DMXFragment)
 {
+	uint16 FinalSendUniverseID = GetFinalSendUniverseID(InUniverseID);
+
+	TSharedPtr<FDMXProtocolUniverseArtNet, ESPMode::ThreadSafe> Universe = UniverseManager->AddUniverseCreate(FinalSendUniverseID);
+
+	if (!Universe->SetDMXFragment(DMXFragment))
+	{
+		return EDMXSendResult::ErrorSetBuffer;
+	}
+
+	EDMXSendResult Result = SendDMXInternal(FinalSendUniverseID, Universe->GetPortID(), Universe->GetOutputDMXBuffer());
+	return Result;
+}
+
+uint16 FDMXProtocolArtNet::GetFinalSendUniverseID(uint16 InUniverseID) const
+{
+	// GlobalSACNUniverseOffset clamp between 0 and 65535(max uint16) 
+	return InUniverseID + GetDefault<UDMXProtocolSettings>()->GlobalArtNetUniverseOffset;
 }
 
 bool FDMXProtocolArtNet::Tick(float DeltaTime)
@@ -99,7 +303,7 @@ bool FDMXProtocolArtNet::Tick(float DeltaTime)
 	return true;
 }
 
-bool FDMXProtocolArtNet::SendDMX(uint16 UniverseID, uint8 PortID, const TSharedPtr<FDMXBuffer>& DMXBuffer) const
+EDMXSendResult FDMXProtocolArtNet::SendDMXInternal(uint16 UniverseID, uint8 PortID, const TSharedPtr<FDMXBuffer>& DMXBuffer) const
 {
 	// Init Packager
 	FDMXProtocolPackager Packager;
@@ -112,17 +316,28 @@ bool FDMXProtocolArtNet::SendDMX(uint16 UniverseID, uint8 PortID, const TSharedP
 	}
 	else
 	{
-		UE_LOG_DMXPROTOCOL(Warning, TEXT("Size of outcoming ArtNet DMX buffer is wrong"));
+		return EDMXSendResult::ErrorSizeBuffer;
 	}
+
+	//Set Packet Data
 	ArtNetDMXPacket.Physical = PortID;
 	ArtNetDMXPacket.Universe = UniverseID;
 	Packager.AddToPackage(&ArtNetDMXPacket);
 
 	// Sending
 	FDMXPacketPtr Packet = MakeShared<FDMXPacket, ESPMode::ThreadSafe>(Packager.GetBuffer());
-	GetSenderInterface()->EnqueueOutboundPackage(Packet);
 
-	return false;
+	if (!GetSenderInterface())
+	{
+		return EDMXSendResult::ErrorEnqueuePackage;
+	}
+
+	if (!GetSenderInterface()->EnqueueOutboundPackage(Packet))
+	{
+		return EDMXSendResult::ErrorEnqueuePackage;
+	}
+
+	return EDMXSendResult::Success;
 }
 
 void FDMXProtocolArtNet::SendRDMCommand(const TSharedPtr<FJsonObject>& CMD)
@@ -151,48 +366,34 @@ bool FDMXProtocolArtNet::TransmitPoll()
 
 bool FDMXProtocolArtNet::TransmitTodRequestToAll()
 {
-	bool bIsSent = false;
-	IDevicesMap DevicesMap;
-	if (GetDeviceManager()->GetDevicesByProtocol(GetProtocolName(), DevicesMap))
+	// Init Packager
+	FDMXProtocolPackager Packager;
+
+	// TransmitTodRequest packet
+	FDMXProtocolArtNetTodRequest TodRequest;
+	for (TMap<uint32, TSharedPtr<FDMXProtocolUniverseArtNet, ESPMode::ThreadSafe>>::TConstIterator UniverseIt = UniverseManager->GetAllUniverses().CreateConstIterator(); UniverseIt; ++UniverseIt)
 	{
-		for (TMap<IDMXProtocolInterface*, TSharedPtr<IDMXProtocolDevice>>::TConstIterator DeviceIt = DevicesMap.CreateConstIterator(); DeviceIt; ++DeviceIt)
+		if (UniverseIt->Value.IsValid())
 		{
-			if (DeviceIt->Value.IsValid())
-			{
-				TransmitTodRequest(DeviceIt->Value);
-				bIsSent = true;
-			}
+			WriteTodRequestAddress(TodRequest, UniverseIt->Value->GetPortAddress());
 		}
 	}
 
-	return bIsSent;
+	Packager.AddToPackage(&TodRequest);
+
+	// Sending
+	FDMXPacketPtr Packet = MakeShared<FDMXPacket, ESPMode::ThreadSafe>(Packager.GetBuffer());
+	return GetSenderInterface()->EnqueueOutboundPackage(Packet);
 }
 
-bool FDMXProtocolArtNet::TransmitTodRequest(const TSharedPtr<IDMXProtocolDevice>& InDevice)
+bool FDMXProtocolArtNet::TransmitTodRequest(uint8 PortAddress)
 {
 	// Init Packager
 	FDMXProtocolPackager Packager;
 
 	// TransmitTodRequest packet
 	FDMXProtocolArtNetTodRequest TodRequest;
-
-	if (InDevice.IsValid())
-	{
-		const IPortsMap* OutputPortsMapDevice = GetPortManager()->GetOutputPortMapByDevice(InDevice.Get());
-		for (uint8 PortIndex = 0; PortIndex < ARTNET_MAX_PORTS; PortIndex++)
-		{
-			if (OutputPortsMapDevice != nullptr)
-			{
-				const TSharedPtr<IDMXProtocolPort>* Port = OutputPortsMapDevice->Find(PortIndex);
-				if (Port != nullptr && Port->IsValid())
-				{
-					FDMXProtocolPortArtNet* PortArtNet = static_cast<FDMXProtocolPortArtNet*>(Port->Get());
-					TodRequest.Address[TodRequest.AdCount] = PortArtNet->GetPortAddress();
-					TodRequest.AdCount++;
-				}
-			}
-		}
-	}
+	WriteTodRequestAddress(TodRequest, PortAddress);
 
 	Packager.AddToPackage(&TodRequest);
 
@@ -201,22 +402,28 @@ bool FDMXProtocolArtNet::TransmitTodRequest(const TSharedPtr<IDMXProtocolDevice>
 	return GetSenderInterface()->EnqueueOutboundPackage(Packet);
 }
 
-bool FDMXProtocolArtNet::TransmitTodData(const TSharedPtr<IDMXProtocolDevice>& InDevice, uint8 InPortID)
+void FDMXProtocolArtNet::WriteTodRequestAddress(FDMXProtocolArtNetTodRequest& TodRequest, uint8 PortAddress)
+{
+	TodRequest.Address[TodRequest.AdCount] = PortAddress;
+	TodRequest.AdCount++;
+}
+
+bool FDMXProtocolArtNet::TransmitTodData(uint32 InUniverseID)
 {
 	// Init Packager
 	FDMXProtocolPackager Packager;
 
 	// TransmitTodData packet
 	FDMXProtocolArtNetTodData TodData;
-	TodData.Port = InPortID;
+	TSharedPtr<FDMXProtocolUniverseArtNet, ESPMode::ThreadSafe> Universe = UniverseManager->GetUniverseById(InUniverseID);
 
-	FDMXProtocolPortArtNet* PortArtNet = static_cast<FDMXProtocolPortArtNet*>(GetPortManager()->GetPortByDeviceAndID(InDevice, InPortID, EDMXPortDirection::DMX_PORT_OUTPUT));
-	if (PortArtNet != nullptr)
+	if (Universe.IsValid())
 	{
-		TodData.Net = PortArtNet->GetNetAddress();
-		TodData.Address = PortArtNet->GetUniverseAddress();
+		TodData.Port = Universe->GetPortID();
+		TodData.Net = Universe->GetNetAddress();
+		TodData.Address = Universe->GetUniverseAddress();
 
-		const TArray<FRDMUID>& TODUID = PortArtNet->GetTODUIDs();
+		const TArray<FRDMUID>& TODUID = Universe->GetTODUIDs();
 		const uint32 NumTODBytes = TODUID.Num() * ARTNET_RDM_UID_WIDTH;
 
 		TodData.UidTotalHi = DMX_SHORT_GET_HIGH_BIT(TODUID.Num());
@@ -237,10 +444,10 @@ bool FDMXProtocolArtNet::TransmitTodData(const TSharedPtr<IDMXProtocolDevice>& I
 	return GetSenderInterface()->EnqueueOutboundPackage(Packet);
 }
 
-bool FDMXProtocolArtNet::TransmitTodControl(const TSharedPtr<IDMXProtocolDevice>& InDevice, uint8 InPortID, uint8 Action)
+bool FDMXProtocolArtNet::TransmitTodControl(uint32 InUniverseID, uint8 Action)
 {
-	// Try to get ArtNet Port
-	FDMXProtocolPortArtNet* PortArtNet = static_cast<FDMXProtocolPortArtNet*>(GetPortManager()->GetPortByDeviceAndID(InDevice, InPortID, EDMXPortDirection::DMX_PORT_OUTPUT));
+	// Try to get Universe
+	TSharedPtr<FDMXProtocolUniverseArtNet, ESPMode::ThreadSafe> Universe = UniverseManager->GetUniverseById(InUniverseID);
 
 	// Init Packager
 	FDMXProtocolPackager Packager;
@@ -249,10 +456,10 @@ bool FDMXProtocolArtNet::TransmitTodControl(const TSharedPtr<IDMXProtocolDevice>
 	FDMXProtocolArtNetTodControl TodControl;
 	TodControl.Cmd = Action;
 
-	if (PortArtNet != nullptr)
+	if (Universe.IsValid())
 	{
-		TodControl.Net = PortArtNet->GetNetAddress();
-		TodControl.Address = PortArtNet->GetPortAddress();
+		TodControl.Net = Universe->GetNetAddress();
+		TodControl.Address = Universe->GetPortAddress();
 	}
 
 	Packager.AddToPackage(&TodControl);
@@ -262,20 +469,20 @@ bool FDMXProtocolArtNet::TransmitTodControl(const TSharedPtr<IDMXProtocolDevice>
 	return GetSenderInterface()->EnqueueOutboundPackage(Packet);
 }
 
-bool FDMXProtocolArtNet::TransmitRDM(const TSharedPtr<IDMXProtocolDevice>& InDevice, uint8 InPortID, const TArray<uint8>& Data)
+bool FDMXProtocolArtNet::TransmitRDM(uint32 InUniverseID, const TArray<uint8>& Data)
 {
-	// Try to get ArtNet Port
-	FDMXProtocolPortArtNet* PortArtNet = static_cast<FDMXProtocolPortArtNet*>(GetPortManager()->GetPortByDeviceAndID(InDevice, InPortID, EDMXPortDirection::DMX_PORT_OUTPUT));
+	// Try to get Universe
+	TSharedPtr<FDMXProtocolUniverseArtNet, ESPMode::ThreadSafe> Universe = UniverseManager->GetUniverseById(InUniverseID);
 
 	// Init Packager
 	FDMXProtocolPackager Packager;
 
 	// TransmitRDM packet
 	FDMXProtocolArtNetRDM RDM;
-	if (PortArtNet != nullptr)
+	if (Universe.IsValid())
 	{
-		RDM.Net = PortArtNet->GetNetAddress();
-		RDM.Address = PortArtNet->GetUniverseAddress();
+		RDM.Net = Universe->GetNetAddress();
+		RDM.Address = Universe->GetUniverseAddress();
 	}
 
 	if (Data.Num() <= ARTNET_MAX_RDM_DATA)
@@ -284,7 +491,7 @@ bool FDMXProtocolArtNet::TransmitRDM(const TSharedPtr<IDMXProtocolDevice>& InDev
 	}
 	else
 	{
-		UE_LOG_DMXPROTOCOL(Warning, TEXT("Size of outcoming ArtNet RDM commands is bigger then limit"));
+		UE_LOG_DMXPROTOCOL(Error, TEXT("%s:Size of outcoming ArtNet RDM commands is bigger then limit"), NetworkErrorMessagePrefix);
 	}
 
 	Packager.AddToPackage(&RDM);
@@ -296,9 +503,14 @@ bool FDMXProtocolArtNet::TransmitRDM(const TSharedPtr<IDMXProtocolDevice>& InDev
 
 void FDMXProtocolArtNet::OnDataReceived(const FArrayReaderPtr& Buffer)
 {
+	FScopeLock Lock(&OnDataReceivedCS);
+
 	// It will be more handlers
 	switch (ArtNet::GetPacketType(Buffer))
 	{
+	case ARTNET_POLL:
+		HandlePool(Buffer);
+		break;
 	case ARTNET_REPLY:
 		HandleReplyPacket(Buffer);
 		break;
@@ -321,11 +533,17 @@ void FDMXProtocolArtNet::OnDataReceived(const FArrayReaderPtr& Buffer)
 	}
 }
 
-bool FDMXProtocolArtNet::HandleReplyPacket(const FArrayReaderPtr & Buffer)
+bool FDMXProtocolArtNet::HandlePool(const FArrayReaderPtr& Buffer)
 {
-	FArtNetPacketReply ArtNetPacketReply;
+	*Buffer << IncomingPollPacket;
+
+	// Nothing to do with packet now
+	return true;
+}
 
-	*Buffer << ArtNetPacketReply;
+bool FDMXProtocolArtNet::HandleReplyPacket(const FArrayReaderPtr & Buffer)
+{
+	*Buffer << PacketReply;
 
 	// Nothing to do with packet now
 	return false;
@@ -337,19 +555,22 @@ bool FDMXProtocolArtNet::HandleDataPacket(const FArrayReaderPtr & Buffer)
 	FDMXProtocolArtNetDMXPacket ArtNetDMXPacket;
 	*Buffer << ArtNetDMXPacket;
 
-	// Write data to input DMX buffer universe
-	IDMXProtocolUniverse* Universe = GetUniverseManager()->GetUniverseById(ArtNetDMXPacket.Universe);
-	if (Universe != nullptr && Universe->GetInputDMXBuffer().IsValid())
+	// Write data to input DMX buffer universe if exists
+	TSharedPtr<FDMXProtocolUniverseArtNet, ESPMode::ThreadSafe> Universe = UniverseManager->GetUniverseById(ArtNetDMXPacket.Universe);
+
+	if (Universe.IsValid() && Universe->GetInputDMXBuffer().IsValid())
 	{
 		// Make sure we copy same amount of data
 		if (Universe->GetInputDMXBuffer()->GetDMXData().Num() == ARTNET_DMX_LENGTH)
 		{
-			FMemory::Memcpy(Universe->GetInputDMXBuffer()->GetDMXData().GetData(), ArtNetDMXPacket.Data, ARTNET_DMX_LENGTH);
+			Universe->GetInputDMXBuffer()->SetDMXBuffer(ArtNetDMXPacket.Data, ARTNET_DMX_LENGTH);
+
+			OnUniverseInputUpdateEvent.Broadcast(GetProtocolName(), ArtNetDMXPacket.Universe, Universe->GetInputDMXBuffer()->GetDMXData());
 			return true;
 		}
 		else
 		{
-			UE_LOG_DMXPROTOCOL(Warning, TEXT("Size of incoming DMX buffer is wrong"));
+			UE_LOG_DMXPROTOCOL(Error, TEXT("%s:Size of incoming DMX buffer is wrong"), NetworkErrorMessagePrefix);
 			return false;
 		}
 	}
diff --git a/Engine/Plugins/Experimental/DigitalMultiplex/DMXProtocol/Source/DMXProtocolArtNet/Private/DMXProtocolArtNet.h b/Engine/Plugins/Experimental/DigitalMultiplex/DMXProtocol/Source/DMXProtocolArtNet/Private/DMXProtocolArtNet.h
new file mode 100644
--- /dev/null
+++ b/Engine/Plugins/Experimental/DigitalMultiplex/DMXProtocol/Source/DMXProtocolArtNet/Private/DMXProtocolArtNet.h
@@ -0,0 +1,161 @@
+// Copyright Epic Games, Inc. All Rights Reserved.
+
+#pragma once
+
+#include "CoreMinimal.h"
+
+#include "Interfaces/IDMXProtocol.h"
+#include "Interfaces/IDMXNetworkInterface.h"
+#include "DMXProtocolArtNetModule.h"
+#include "Packets/DMXProtocolArtNetPackets.h"
+
+#include "HAL/ThreadSafeCounter.h"
+#include "Interfaces/IPv4/IPv4Endpoint.h"
+#include "HAL/CriticalSection.h"
+
+class FSocket;
+class FDMXProtocolUniverseArtNet;
+
+class DMXPROTOCOLARTNET_API FDMXProtocolArtNet
+	: public IDMXProtocol
+	, public IDMXNetworkInterface
+{
+public:
+	explicit FDMXProtocolArtNet(const FName& InProtocolName, const FJsonObject& InSettings);
+
+	//~ Begin IDMXProtocolBase implementation
+	virtual bool Init() override;
+	virtual bool Shutdown() override;
+	virtual bool Tick(float DeltaTime) override;
+	//~ End IDMXProtocolBase implementation
+
+	//~ Begin IDMXProtocol implementation
+	virtual const FName& GetProtocolName() const override;
+	virtual TSharedPtr<FJsonObject> GetSettings() const override;
+	virtual TSharedPtr<IDMXProtocolSender> GetSenderInterface() const override;
+	virtual EDMXSendResult SendDMXFragment(uint16 UniverseID, const IDMXFragmentMap& DMXFragment) override;
+	virtual EDMXSendResult SendDMXFragmentCreate(uint16 InUniverseID, const IDMXFragmentMap& DMXFragment) override;
+	virtual uint16 GetFinalSendUniverseID(uint16 InUniverseID) const override;
+	virtual bool IsEnabled() const override;
+	virtual TSharedPtr<IDMXProtocolUniverse, ESPMode::ThreadSafe> AddUniverse(const FJsonObject& InSettings) override;
+	virtual void CollectUniverses(const TArray<FDMXUniverse>& Universes) override;
+	virtual bool RemoveUniverseById(uint32 InUniverseId) override;
+	virtual void RemoveAllUniverses() override;
+	virtual TSharedPtr<IDMXProtocolUniverse, ESPMode::ThreadSafe> GetUniverseById(uint32 InUniverseId) const override;
+	virtual uint32 GetUniversesNum() const override;
+	virtual uint16 GetMinUniverseID() const override;
+	virtual uint16 GetMaxUniverses() const override;
+	virtual FOnUniverseInputUpdateEvent& GetOnUniverseInputUpdate() override { return OnUniverseInputUpdateEvent; }
+	//~ End IDMXProtocol implementation
+
+	//~ Begin IDMXNetworkInterface implementation
+	virtual void OnNetworkInterfaceChanged(const FString& InInterfaceIPAddress) override;
+	virtual bool RestartNetworkInterface(const FString& InInterfaceIPAddress, FString& OutErrorMessage) override;
+	virtual void ReleaseNetworkInterface() override;
+	//~ End IDMXNetworkInterface implementation
+
+	//~ Begin IDMXProtocolRDM implementation
+	virtual void SendRDMCommand(const TSharedPtr<FJsonObject>& CMD) override;
+	virtual void RDMDiscovery(const TSharedPtr<FJsonObject>& CMD) override;
+	//~ End IDMXProtocol implementation
+
+	//~ ArtNet Transmit functions
+	bool TransmitPoll();
+	bool TransmitTodRequestToAll();
+	bool TransmitTodRequest(uint8 PortAddress);
+	void WriteTodRequestAddress(FDMXProtocolArtNetTodRequest& TodRequest, uint8 PortAddress);
+
+	//~ ArtNet Getters functions
+	const FDMXProtocolArtNetPollPacket& GetIncomingPoll() const { return IncomingPollPacket; }
+	const FDMXProtocolArtNetPacketReply& GetIncomingPacketReply() const { return PacketReply; }
+	const FDMXProtocolArtNetTodRequest& GetIncomingTodRequest() const { return IncomingTodRequest; }
+	const FDMXProtocolArtNetTodData& GetIncomingTodData() const { return IncomingTodData; }
+	const FDMXProtocolArtNetTodControl& GetIncomingTodControl() const { return IncomingTodControl; }
+	const FDMXProtocolArtNetRDM& GetIncomingRDM() const { return IncomingRDM; }
+
+	/**
+	 * Transmit TOD data
+	 * @param InUniverseID The universe which we want to transmit
+	 * @return true if the call was successfully placed for sending queue.
+	 */
+	bool TransmitTodData(uint32 InUniverseID);
+
+	/**
+	 * Transmit TOD cocontrol package
+	 * @param InUniverseID The universe which we want to transmit
+	 * @param Action tod control action
+	 * @return true if the call was successfully placed for sending queue.
+	 */
+	bool TransmitTodControl(uint32 InUniverseID, uint8 Action);
+
+	/**
+	 * Send a RDM message
+	 * @param InUniverseID The universe which we want to transmit
+	 * @param Data Buffer array to send
+	 * @return true if the call was successfully placed for sending queue.
+	 */
+	bool TransmitRDM(uint32 InUniverseID, const TArray<uint8>& Data);
+
+private:
+	EDMXSendResult SendDMXInternal(uint16 InUniverseID, uint8 InPort, const TSharedPtr<FDMXBuffer>& DMXBuffer) const;
+
+	//~ Only the factory makes instances
+	FDMXProtocolArtNet() = delete;
+
+public:
+	const TSharedPtr<FInternetAddr>& GetBroadcastAddr() { return BroadcastAddr; }
+	void OnDataReceived(const FArrayReaderPtr & Buffer);
+
+private:
+	//~ DMX Handlers
+	bool HandlePool(const FArrayReaderPtr& Buffer);
+	bool HandleReplyPacket(const FArrayReaderPtr & Buffer);
+	bool HandleDataPacket(const FArrayReaderPtr & Buffer);
+
+	//~ RDM Handlers
+	bool HandleTodRequest(const FArrayReaderPtr & Buffer);
+	bool HandleTodData(const FArrayReaderPtr & Buffer);
+	bool HandleTodControl(const FArrayReaderPtr & Buffer);
+	bool HandleRdm(const FArrayReaderPtr & Buffer);
+
+private:
+	FName ProtocolName;
+	TSharedPtr<FJsonObject> Settings;
+
+	TSharedPtr<FDMXProtocolUniverseManager<FDMXProtocolUniverseArtNet>> UniverseManager;
+
+	TSharedPtr<IDMXProtocolSender> ArtNetSender;
+	TSharedPtr<IDMXProtocolReceiver> ArtNetReceiver;
+
+	mutable FCriticalSection OnDataReceivedCS;
+
+	/** Holds the network socket used to transport packages. */
+	FSocket* BroadcastSocket;
+
+	FSocket* ListeningSocket;
+
+	TSharedPtr<FInternetAddr> BroadcastAddr;
+	FIPv4Endpoint BroadcastEndpoint;
+
+	TSharedPtr<FInternetAddr> SenderAddr;
+	FIPv4Endpoint SenderEndpoint;
+
+	FDMXProtocolArtNetPollPacket IncomingPollPacket;
+	FDMXProtocolArtNetPacketReply PacketReply;
+	FDMXProtocolArtNetTodRequest IncomingTodRequest;
+	FDMXProtocolArtNetTodData IncomingTodData;
+	FDMXProtocolArtNetTodControl IncomingTodControl;
+	FDMXProtocolArtNetRDM IncomingRDM;
+
+	FOnUniverseInputUpdateEvent OnUniverseInputUpdateEvent;
+
+	/** Mutex protecting access to the sockets. */
+	mutable FCriticalSection SocketsCS;
+
+	FString InterfaceIPAddress;
+
+	FDelegateHandle NetworkInterfaceChangedHandle;
+
+	const TCHAR* NetworkErrorMessagePrefix;
+};
+
diff --git a/Engine/Plugins/Experimental/DigitalMultiplex/DMXProtocol/Source/DMXProtocolArtNet/Private/DMXProtocolArtNetConstants.h b/Engine/Plugins/Experimental/DigitalMultiplex/DMXProtocol/Source/DMXProtocolArtNet/Private/DMXProtocolArtNetConstants.h
--- a/Engine/Plugins/Experimental/DigitalMultiplex/DMXProtocol/Source/DMXProtocolArtNet/Private/DMXProtocolArtNetConstants.h
+++ b/Engine/Plugins/Experimental/DigitalMultiplex/DMXProtocol/Source/DMXProtocolArtNet/Private/DMXProtocolArtNetConstants.h
@@ -19,6 +19,8 @@ enum
 #define DMX_SHORT_GET_LOW_BYTE(x)  (LOW_BYTE & x)
 #endif
 
+#define DMX_PROTOCOLNAME_ARTNET "Art-Net"
+
 enum
 {
 	ARTNET_LONG_NAME_LENGTH = 64,
@@ -103,6 +105,7 @@ enum
 {
 	ARTNET_ADDRESS_NO_CHANGE = 0x7f,
 	ARTNET_PORT = 6454,
+	ARTNET_SENDER_PORT = 3330,
 	ARTNET_STRING_SIZE = 8,
 	ARTNET_VERSION = 14,
 	OEM_HI = 0x04,
diff --git a/Engine/Plugins/Experimental/DigitalMultiplex/DMXProtocol/Source/DMXProtocolArtNet/Private/DMXProtocolArtNetModule.cpp b/Engine/Plugins/Experimental/DigitalMultiplex/DMXProtocol/Source/DMXProtocolArtNet/Private/DMXProtocolArtNetModule.cpp
--- a/Engine/Plugins/Experimental/DigitalMultiplex/DMXProtocol/Source/DMXProtocolArtNet/Private/DMXProtocolArtNetModule.cpp
+++ b/Engine/Plugins/Experimental/DigitalMultiplex/DMXProtocol/Source/DMXProtocolArtNet/Private/DMXProtocolArtNetModule.cpp
@@ -7,11 +7,12 @@
 
 #include "DMXProtocolArtNet.h"
 #include "DMXProtocolTypes.h"
+#include "DMXProtocolArtNetConstants.h"
 
 
 IMPLEMENT_MODULE(FDMXProtocolArtNetModule, DMXProtocolArtNet);
 
-const FName FDMXProtocolArtNetModule::NAME_Artnet = FName(TEXT("Art-Net"));
+const FName FDMXProtocolArtNetModule::NAME_Artnet = FName(DMX_PROTOCOLNAME_ARTNET);
 
 FAutoConsoleCommand FDMXProtocolArtNetModule::SendDMXCommand(
 	TEXT("DMX.ArtNet.SendDMX"),
@@ -27,14 +28,14 @@ TSharedPtr<IDMXProtocol> FDMXProtocolFactoryArtNet::CreateProtocol(const FName &
 	{
 		if (!ProtocolArtNetPtr->Init())
 		{
-			UE_LOG_DMXPROTOCOL(Warning, TEXT("ArtNet failed to initialize!"));
+			UE_LOG_DMXPROTOCOL(Verbose, TEXT("ArtNet failed to initialize!"));
 			ProtocolArtNetPtr->Shutdown();
 			ProtocolArtNetPtr = nullptr;
 		}
 	}
 	else
 	{
-		UE_LOG_DMXPROTOCOL(Warning, TEXT("ArtNet disabled!"));
+		UE_LOG_DMXPROTOCOL(Verbose, TEXT("ArtNet disabled!"));
 		ProtocolArtNetPtr->Shutdown();
 		ProtocolArtNetPtr = nullptr;
 	}
@@ -72,7 +73,7 @@ void FDMXProtocolArtNetModule::SendDMXCommandHandler(const TArray<FString>& Args
 {
 	if (Args.Num() < 2)
 	{
-		UE_LOG_DMXPROTOCOL(Warning, TEXT("Not enough arguments. It won't be sent\n"
+		UE_LOG_DMXPROTOCOL(Verbose, TEXT("Not enough arguments. It won't be sent\n"
 										"Command structure is DMX.ArtNet.SendDMX [UniverseID] Channel:Value Channel:Value Channel:Value\n"
 										"For example: DMX.ArtNet.SendDMX 17 10:6 11:7 12:8 13:9"));
 		return;
@@ -82,7 +83,7 @@ void FDMXProtocolArtNetModule::SendDMXCommandHandler(const TArray<FString>& Args
 	LexTryParseString<uint32>(UniverseID, *Args[0]);
 	if (UniverseID > ARTNET_MAX_UNIVERSES)
 	{
-		UE_LOG_DMXPROTOCOL(Warning, TEXT("The UniverseID is bigger then the max universe value. It won't be sent\n"
+		UE_LOG_DMXPROTOCOL(Verbose, TEXT("The UniverseID is bigger then the max universe value. It won't be sent\n"
 			"For example: DMX.ArtNet.SendDMX 17 10:6 11:7 12:8 13:9\n"
 			"Where Universe %d should be less then %d"), UniverseID, ARTNET_MAX_UNIVERSES);
 		return;
@@ -102,7 +103,7 @@ void FDMXProtocolArtNetModule::SendDMXCommandHandler(const TArray<FString>& Args
 		LexTryParseString<uint32>(Key, *KeyStr);
 		if (Key > DMX_UNIVERSE_SIZE)
 		{
-			UE_LOG_DMXPROTOCOL(Warning, TEXT("The input channel is bigger then the universe size. It won't be sent\n"
+			UE_LOG_DMXPROTOCOL(Verbose, TEXT("The input channel is bigger then the universe size. It won't be sent\n"
 				"For example: DMX.ArtNet.SendDMX 17 10:6 11:7 12:8 13:9\n"
 				"Where channel %d should be less then %d"), Key, DMX_UNIVERSE_SIZE);
 			return;
@@ -110,7 +111,7 @@ void FDMXProtocolArtNetModule::SendDMXCommandHandler(const TArray<FString>& Args
 		LexTryParseString<uint32>(Value, *ValueStr);
 		if (Value > DMX_MAX_CHANNEL_VALUE)
 		{
-			UE_LOG_DMXPROTOCOL(Warning, TEXT("The input value is bigger then the universe size. It won't be sent\n"
+			UE_LOG_DMXPROTOCOL(Verbose, TEXT("The input value is bigger then the universe size. It won't be sent\n"
 				"For example: DMX.ArtNet.SendDMX 17 10:6 11:7 12:8 13:9\n"
 				"Where value %d should be less then %d"), Value, DMX_MAX_CHANNEL_VALUE);
 			return;
@@ -118,6 +119,9 @@ void FDMXProtocolArtNetModule::SendDMXCommandHandler(const TArray<FString>& Args
 		DMXFragment.Add(Key, Value);
 	}
 
-	FDMXProtocolArtNet* DMXProtocol = IDMXProtocol::Get<FDMXProtocolArtNet>(FDMXProtocolArtNetModule::NAME_Artnet);
-	DMXProtocol->SetDMXFragment(UniverseID, DMXFragment, true);
+	TSharedPtr<IDMXProtocol> DMXProtocol = IDMXProtocol::Get(FDMXProtocolArtNetModule::NAME_Artnet);
+	if (DMXProtocol.IsValid())
+	{
+		DMXProtocol->SendDMXFragmentCreate(UniverseID, DMXFragment);
+	}
 }
diff --git a/Engine/Plugins/Experimental/DigitalMultiplex/DMXProtocol/Source/DMXProtocolArtNet/Private/DMXProtocolArtNetUtils.cpp b/Engine/Plugins/Experimental/DigitalMultiplex/DMXProtocol/Source/DMXProtocolArtNet/Private/DMXProtocolArtNetUtils.cpp
--- a/Engine/Plugins/Experimental/DigitalMultiplex/DMXProtocol/Source/DMXProtocolArtNet/Private/DMXProtocolArtNetUtils.cpp
+++ b/Engine/Plugins/Experimental/DigitalMultiplex/DMXProtocol/Source/DMXProtocolArtNet/Private/DMXProtocolArtNetUtils.cpp
@@ -1,7 +1,6 @@
 // Copyright Epic Games, Inc. All Rights Reserved.
 
 #include "DMXProtocolArtNetUtils.h"
-#include "DMXProtocolArtNetPublicUtils.h"
 
 namespace ArtNet
 {
@@ -21,13 +20,4 @@ namespace ArtNet
 
 		return OpCode;
 	}
-
-	uint16 ComputeUniverseID(const uint8 Net, const uint8 Subnet, const uint8 Universe)
-	{
-		/* Bits 15 			Bits 14-8        | Bits 7-4      | Bits 3-0
-		* 0 				Net 			 | Sub-Net 		 | Universe
-		* 0				    (0b1111111 << 8) | (0b1111 << 4) | (0b1111)
-		*/
-		return (Net << 8) | (Subnet << 4) | Universe;
-	}
 }
diff --git a/Engine/Plugins/Experimental/DigitalMultiplex/DMXProtocol/Source/DMXProtocolArtNet/Private/DMXProtocolTransportArtNet.cpp b/Engine/Plugins/Experimental/DigitalMultiplex/DMXProtocol/Source/DMXProtocolArtNet/Private/DMXProtocolTransportArtNet.cpp
--- a/Engine/Plugins/Experimental/DigitalMultiplex/DMXProtocol/Source/DMXProtocolArtNet/Private/DMXProtocolTransportArtNet.cpp
+++ b/Engine/Plugins/Experimental/DigitalMultiplex/DMXProtocol/Source/DMXProtocolArtNet/Private/DMXProtocolTransportArtNet.cpp
@@ -99,7 +99,6 @@ bool FDMXProtocolSenderArtNet::EnqueueOutboundPackage(FDMXPacketPtr Packet)
 
 	WorkEvent->Trigger();
 
-
 	return true;
 }
 
@@ -194,19 +193,19 @@ void FDMXProtocolReceiverArtNet::Update(const FTimespan& SocketWaitTime)
 		return;
 	}
 
-	TSharedRef<FInternetAddr> Sender = SocketSubsystem->CreateInternetAddr();
+    TSharedRef<FInternetAddr> Sender = SocketSubsystem->CreateInternetAddr();
 	uint32 Size;
 
 	while (Socket->HasPendingData(Size))
 	{
-		FArrayReaderPtr Reader = MakeShareable(new FArrayReader(true));
+		FArrayReaderPtr Reader = MakeShared<FArrayReader, ESPMode::ThreadSafe>(true);
 		Reader->SetNumUninitialized(FMath::Min(Size, 65507u));
 
 		int32 Read = 0;
 
-		if (Socket->RecvFrom(Reader->GetData(), Reader->Num(), Read, *Sender))
+        if (Socket->RecvFrom(Reader->GetData(), Reader->Num(), Read, *Sender))
 		{
-			Reader->RemoveAt(Read, Reader->Num() - Read, false);
+            Reader->RemoveAt(Read, Reader->Num() - Read, false);
 			DMXDataReceiveDelegate.ExecuteIfBound(Reader);
 		}
 	}
diff --git a/Engine/Plugins/Experimental/DigitalMultiplex/DMXProtocol/Source/DMXProtocolArtNet/Private/DMXProtocolTransportArtNet.h b/Engine/Plugins/Experimental/DigitalMultiplex/DMXProtocol/Source/DMXProtocolArtNet/Private/DMXProtocolTransportArtNet.h
--- a/Engine/Plugins/Experimental/DigitalMultiplex/DMXProtocol/Source/DMXProtocolArtNet/Private/DMXProtocolTransportArtNet.h
+++ b/Engine/Plugins/Experimental/DigitalMultiplex/DMXProtocol/Source/DMXProtocolArtNet/Private/DMXProtocolTransportArtNet.h
@@ -3,10 +3,13 @@
 #pragma once
 
 #include "CoreMinimal.h"
-#include "Containers/Queue.h"
+
 #include "Interfaces/IDMXProtocolTransport.h"
 #include "DMXProtocolTypes.h"
 
+#include "Containers/Queue.h"
+#include "HAL/CriticalSection.h"
+
 class FSocket;
 class ISocketSubsystem;
 class FDMXProtocolArtNet;
diff --git a/Engine/Plugins/Experimental/DigitalMultiplex/DMXProtocol/Source/DMXProtocolArtNet/Private/DMXProtocolUniverseArtNet.cpp b/Engine/Plugins/Experimental/DigitalMultiplex/DMXProtocol/Source/DMXProtocolArtNet/Private/DMXProtocolUniverseArtNet.cpp
--- a/Engine/Plugins/Experimental/DigitalMultiplex/DMXProtocol/Source/DMXProtocolArtNet/Private/DMXProtocolUniverseArtNet.cpp
+++ b/Engine/Plugins/Experimental/DigitalMultiplex/DMXProtocol/Source/DMXProtocolArtNet/Private/DMXProtocolUniverseArtNet.cpp
@@ -4,28 +4,45 @@
 #include "Dom/JsonObject.h"
 
 #include "Interfaces/IDMXProtocol.h"
-#include "Interfaces/IDMXProtocolPort.h"
-#include "Interfaces/IDMXProtocolDevice.h"
 #include "DMXProtocolConstants.h"
 #include "DMXProtocolTypes.h"
 
-FDMXProtocolUniverseArtNet::FDMXProtocolUniverseArtNet(IDMXProtocol* InDMXProtocol, TSharedPtr<IDMXProtocolPort> InPort)
-	: DMXProtocol(InDMXProtocol)
-	, Port(InPort)
-	, Priority(0)
-	, UniverseID(InPort->GetUniverseID())
+//~ Artnet specific constants
+static const uint8 HIGH_NIBBLE = 0xF0;
+static const uint8 LOW_NIBBLE = 0x0F;
+static const uint8 MAX_NET = 0x7f;
+
+FDMXProtocolUniverseArtNet::FDMXProtocolUniverseArtNet(TSharedPtr<IDMXProtocol> InDMXProtocol, const FJsonObject& InSettings)
+	: WeakDMXProtocol(InDMXProtocol)
 {
-	checkf(DMXProtocol, TEXT("DMXProtocol pointer is nullptr"));
-	checkf(Port.IsValid(), TEXT("DMXProtocol port is not valid"));
+	checkf(WeakDMXProtocol.IsValid(), TEXT("DMXProtocol pointer is nullptr"));
 
-	InputDMXBuffer = MakeShared<FDMXBuffer>();
+	Priority = 100;
 	OutputDMXBuffer = MakeShared<FDMXBuffer>();
+	InputDMXBuffer = MakeShared<FDMXBuffer>();
+
+	Settings = MakeShared<FJsonObject>(InSettings);
+	checkf(Settings->HasField(TEXT("PortID")), TEXT("DMXProtocol PortID is not valid"));
+	checkf(Settings->HasField(TEXT("UniverseID")), TEXT("DMXProtocol Universe is not valid"));
+	PortID = Settings->GetNumberField(TEXT("PortID"));
+
+	/* Bits 15 			Bits 14-8        | Bits 7-4      | Bits 3-0
+	* 0 				Net 			 | Sub-Net 		 | Universe
+	* 0				    (0b1111111 << 8) | (0b1111 << 4) | (0b111)
+	*/
+	UniverseID = Settings->GetNumberField(TEXT("UniverseID"));
+	uint8 Net = (uint8)((UniverseID & 0b0111111100000000) >> 8);
+	uint8 Subnet = (uint8)((UniverseID & 0b11110000) >> 4);
+	uint8 Universe = (uint8)((UniverseID & 0b1111));
+	SetNetAddress(Net);
+	SetSubnetAddress(Subnet);
+	SetUniverse(Universe);
 }
 
 
-IDMXProtocol * FDMXProtocolUniverseArtNet::GetProtocol() const
+TSharedPtr<IDMXProtocol> FDMXProtocolUniverseArtNet::GetProtocol() const
 {
-	return DMXProtocol;
+	return WeakDMXProtocol.Pin();
 }
 
 TSharedPtr<FDMXBuffer> FDMXProtocolUniverseArtNet::GetOutputDMXBuffer() const
@@ -40,8 +57,7 @@ TSharedPtr<FDMXBuffer> FDMXProtocolUniverseArtNet::GetInputDMXBuffer() const
 
 bool FDMXProtocolUniverseArtNet::SetDMXFragment(const IDMXFragmentMap & DMXFragment)
 {
-	OutputDMXBuffer->SetDMXFragment(DMXFragment);
-	return false;
+	return OutputDMXBuffer->SetDMXFragment(DMXFragment);
 }
 
 uint8 FDMXProtocolUniverseArtNet::GetPriority() const
@@ -49,12 +65,99 @@ uint8 FDMXProtocolUniverseArtNet::GetPriority() const
 	return Priority;
 }
 
-uint16 FDMXProtocolUniverseArtNet::GetUniverseID() const
+uint32 FDMXProtocolUniverseArtNet::GetUniverseID() const
 {
 	return UniverseID;
 }
 
-TWeakPtr<IDMXProtocolPort> FDMXProtocolUniverseArtNet::GetCachedUniversePort() const
+TSharedPtr<FJsonObject> FDMXProtocolUniverseArtNet::GetSettings() const
+{
+	return Settings;
+}
+
+bool FDMXProtocolUniverseArtNet::IsSupportRDM() const
+{
+	return true;
+}
+
+uint8 FDMXProtocolUniverseArtNet::GetPortID()
+{
+	return PortID;
+}
+
+bool FDMXProtocolUniverseArtNet::SetNetAddress(uint8 InNet)
+{
+	if (InNet & 0x80)
+	{
+		UE_LOG_DMXPROTOCOL(Verbose, TEXT("Art-Net net address > 127, truncating!"));
+		NetAddress = InNet & MAX_NET;
+		return true;
+	}
+	if (NetAddress == InNet)
+	{
+		return false;
+	}
+
+	NetAddress = InNet;
+	return true;
+}
+
+bool FDMXProtocolUniverseArtNet::SetSubnetAddress(uint8 InSubnetAddress)
+{
+	// Bits 7-4
+	// Shift for 4 bytes
+	SubnetAddress = InSubnetAddress << 4;
+	if (SubnetAddress == (PortAddress & HIGH_NIBBLE))
+	{
+		return false;
+	}
+
+	PortAddress = SubnetAddress | (PortAddress & LOW_NIBBLE);
+
+	return true;
+}
+
+bool FDMXProtocolUniverseArtNet::SetUniverse(uint8 InUniverse)
+{
+	// Bits 1-3
+	UniverseAddress = InUniverse & LOW_NIBBLE;
+	if ((PortAddress & LOW_NIBBLE) == UniverseAddress)
+	{
+		return false;
+	}
+
+	PortAddress = ((PortAddress & HIGH_NIBBLE) | UniverseAddress);
+
+	return true;
+}
+
+uint8 FDMXProtocolUniverseArtNet::GetPortAddress() const
+{
+	return PortAddress;
+}
+
+uint8 FDMXProtocolUniverseArtNet::GetNetAddress() const
+{
+	return NetAddress;
+}
+
+uint8 FDMXProtocolUniverseArtNet::GetSubnetAddress() const
+{
+	return SubnetAddress;
+}
+
+uint8 FDMXProtocolUniverseArtNet::GetUniverseAddress() const
 {
-	return Port;
+	return UniverseAddress;
 }
+
+const TArray<FRDMUID>& FDMXProtocolUniverseArtNet::GetTODUIDs() const
+{
+	return TODUIDs;
+}
+
+void FDMXProtocolUniverseArtNet::AddTODUID(const FRDMUID& InUID)
+{
+	TODUIDs.Add(InUID);
+}
+
diff --git a/Engine/Plugins/Experimental/DigitalMultiplex/DMXProtocol/Source/DMXProtocolArtNet/Private/DMXProtocolUniverseArtNet.h b/Engine/Plugins/Experimental/DigitalMultiplex/DMXProtocol/Source/DMXProtocolArtNet/Private/DMXProtocolUniverseArtNet.h
--- a/Engine/Plugins/Experimental/DigitalMultiplex/DMXProtocol/Source/DMXProtocolArtNet/Private/DMXProtocolUniverseArtNet.h
+++ b/Engine/Plugins/Experimental/DigitalMultiplex/DMXProtocol/Source/DMXProtocolArtNet/Private/DMXProtocolUniverseArtNet.h
@@ -5,27 +5,83 @@
 #include "CoreMinimal.h"
 #include "Interfaces/IDMXProtocolUniverse.h"
 
+struct FRDMUID;
+
 class DMXPROTOCOLARTNET_API FDMXProtocolUniverseArtNet
 	: public IDMXProtocolUniverse
 {
 public:
-	FDMXProtocolUniverseArtNet(IDMXProtocol* InDMXProtocol, TSharedPtr<IDMXProtocolPort> InPort);
+	FDMXProtocolUniverseArtNet(TSharedPtr<IDMXProtocol> InDMXProtocol, const FJsonObject& InSettings);
 
 	//~ Begin IDMXProtocolDevice implementation
-	virtual IDMXProtocol* GetProtocol() const override;
-	virtual TSharedPtr<FDMXBuffer> GetOutputDMXBuffer() const override;
+	virtual TSharedPtr<IDMXProtocol> GetProtocol() const override;
 	virtual TSharedPtr<FDMXBuffer> GetInputDMXBuffer() const override;
+	virtual TSharedPtr<FDMXBuffer> GetOutputDMXBuffer() const override;
 	virtual bool SetDMXFragment(const IDMXFragmentMap& DMXFragment) override;
 	virtual uint8 GetPriority() const override;
-	virtual uint16 GetUniverseID() const override;
-	virtual TWeakPtr<IDMXProtocolPort> GetCachedUniversePort() const override;
+	virtual uint32 GetUniverseID() const override;
+	virtual TSharedPtr<FJsonObject> GetSettings() const override;
+	virtual bool IsSupportRDM() const override;
 	//~ End IDMXProtocolDevice implementation
 
+public:
+	/**
+	 * Represent physical ArtNet controller port ID. Values from 0 to 3
+	 * @return The 8-bit port address
+	 */
+	uint8 GetPortID();
+
+	/**
+	 * The 8-bit port address, which is made up of the sub-net and universe
+	 * @return The 8-bit port address
+	 */
+	uint8 GetPortAddress() const;
+	//~ End IDMXProtocolPort implementation
+
+	/**
+	 * Bits 14-8
+	 * @return The 8-bit net address
+	 */
+	uint8 GetNetAddress() const;
+
+	/**
+	 * Bits 7-4
+	 * @return The 8-bit subnet address
+	 */
+	uint8 GetSubnetAddress() const;
+
+	/**
+	 * Bits 3-0
+	 * @return The 8-bit uni address
+	 */
+	uint8 GetUniverseAddress() const;
+
+	const TArray<FRDMUID>& GetTODUIDs() const;
+
+	void AddTODUID(const FRDMUID& InUID);
+
+private:
+
+	//~ ArtNet specific implemntation
+	bool SetNetAddress(uint8 InNet);
+	bool SetSubnetAddress(uint8 InSubnetAddress);
+	bool SetUniverse(uint8 InUniverse);
+
 private:
-	IDMXProtocol* DMXProtocol;
-	TWeakPtr<IDMXProtocolPort> Port;
+	TWeakPtr<IDMXProtocol> WeakDMXProtocol;
 	TSharedPtr<FDMXBuffer> OutputDMXBuffer;
 	TSharedPtr<FDMXBuffer> InputDMXBuffer;
 	uint8 Priority;
 	uint16 UniverseID;
+
+	uint8 PortID;
+	uint8 PortAddress;
+	uint8 NetAddress;
+	uint8 SubnetAddress;
+	uint8 UniverseAddress;
+	bool bIsRDMSupport;
+	TSharedPtr<FJsonObject> Settings;
+
+	// Array of TODs UID. This is using for RDM discovery
+	TArray<FRDMUID> TODUIDs;
 };
\ No newline at end of file
diff --git a/Engine/Plugins/Experimental/DigitalMultiplex/DMXProtocol/Source/DMXProtocolArtNet/Private/Packets/DMXProtocolArtNetPackets.cpp b/Engine/Plugins/Experimental/DigitalMultiplex/DMXProtocol/Source/DMXProtocolArtNet/Private/Packets/DMXProtocolArtNetPackets.cpp
--- a/Engine/Plugins/Experimental/DigitalMultiplex/DMXProtocol/Source/DMXProtocolArtNet/Private/Packets/DMXProtocolArtNetPackets.cpp
+++ b/Engine/Plugins/Experimental/DigitalMultiplex/DMXProtocol/Source/DMXProtocolArtNet/Private/Packets/DMXProtocolArtNetPackets.cpp
@@ -2,12 +2,13 @@
 
 #include "Packets/DMXProtocolArtNetPackets.h"
 #include "DMXProtocolTypes.h"
-#include "Serialization/DMXMemoryWriter.h"
 #include "DMXProtocolMacros.h"
 
+#include "Serialization/BufferArchive.h"
+
 REGISTER_DMX_ARCHIVE(FDMXProtocolArtNetDMXPacket);
 REGISTER_DMX_ARCHIVE(FDMXProtocolArtNetPollPacket);
-REGISTER_DMX_ARCHIVE(FArtNetPacketReply);
+REGISTER_DMX_ARCHIVE(FDMXProtocolArtNetPacketReply);
 REGISTER_DMX_ARCHIVE(FDMXProtocolArtNetTodRequest);
 REGISTER_DMX_ARCHIVE(FDMXProtocolArtNetTodData);
 REGISTER_DMX_ARCHIVE(FDMXProtocolArtNetTodControl);
@@ -15,7 +16,7 @@ REGISTER_DMX_ARCHIVE(FDMXProtocolArtNetRDM);
 
 TSharedPtr<FBufferArchive> FDMXProtocolArtNetDMXPacket::Pack()
 {
-	TSharedPtr<FDMXMemoryWriter> Writer = MakeShared<FDMXMemoryWriter>();
+	TSharedPtr<FBufferArchive> Writer = MakeShared<FBufferArchive>();
 	*Writer << *this;
 
 	return Writer;
@@ -38,7 +39,7 @@ void FDMXProtocolArtNetDMXPacket::Serialize(FArchive & Ar)
 
 TSharedPtr<FBufferArchive> FDMXProtocolArtNetPollPacket::Pack()
 {
-	TSharedPtr<FDMXMemoryWriter> Writer = MakeShared<FDMXMemoryWriter>();
+	TSharedPtr<FBufferArchive> Writer = MakeShared<FBufferArchive>();
 	*Writer << *this;
 
 	return Writer;
@@ -54,7 +55,7 @@ void FDMXProtocolArtNetPollPacket::Serialize(FArchive & Ar)
 	Ar << Priority;
 }
 
-void FArtNetPacketReply::Serialize(FArchive & Ar)
+void FDMXProtocolArtNetPacketReply::Serialize(FArchive & Ar)
 {
 	Ar.Serialize((void*)ID, ARTNET_STRING_SIZE);
 	Ar << OpCode;
@@ -92,7 +93,7 @@ void FArtNetPacketReply::Serialize(FArchive & Ar)
 
 TSharedPtr<FBufferArchive> FDMXProtocolArtNetTodRequest::Pack()
 {
-	TSharedPtr<FDMXMemoryWriter> Writer = MakeShared<FDMXMemoryWriter>();
+	TSharedPtr<FBufferArchive> Writer = MakeShared<FBufferArchive>();
 	*Writer << *this;
 
 	return Writer;
@@ -121,7 +122,7 @@ void FDMXProtocolArtNetTodRequest::Serialize(FArchive & Ar)
 
 TSharedPtr<FBufferArchive> FDMXProtocolArtNetTodData::Pack()
 {
-	TSharedPtr<FDMXMemoryWriter> Writer = MakeShared<FDMXMemoryWriter>();
+	TSharedPtr<FBufferArchive> Writer = MakeShared<FBufferArchive>();
 	*Writer << *this;
 
 	return Writer;
@@ -154,7 +155,7 @@ void FDMXProtocolArtNetTodData::Serialize(FArchive & Ar)
 
 TSharedPtr<FBufferArchive> FDMXProtocolArtNetTodControl::Pack()
 {
-	TSharedPtr<FDMXMemoryWriter> Writer = MakeShared<FDMXMemoryWriter>();
+	TSharedPtr<FBufferArchive> Writer = MakeShared<FBufferArchive>();
 	*Writer << *this;
 
 	return Writer;
@@ -182,7 +183,7 @@ void FDMXProtocolArtNetTodControl::Serialize(FArchive & Ar)
 
 TSharedPtr<FBufferArchive> FDMXProtocolArtNetRDM::Pack()
 {
-	TSharedPtr<FDMXMemoryWriter> Writer = MakeShared<FDMXMemoryWriter>();
+	TSharedPtr<FBufferArchive> Writer = MakeShared<FBufferArchive>();
 	*Writer << *this;
 
 	return Writer;
diff --git a/Engine/Plugins/Experimental/DigitalMultiplex/DMXProtocol/Source/DMXProtocolArtNet/Private/Packets/DMXProtocolArtNetPackets.h b/Engine/Plugins/Experimental/DigitalMultiplex/DMXProtocol/Source/DMXProtocolArtNet/Private/Packets/DMXProtocolArtNetPackets.h
--- a/Engine/Plugins/Experimental/DigitalMultiplex/DMXProtocol/Source/DMXProtocolArtNet/Private/Packets/DMXProtocolArtNetPackets.h
+++ b/Engine/Plugins/Experimental/DigitalMultiplex/DMXProtocol/Source/DMXProtocolArtNet/Private/Packets/DMXProtocolArtNetPackets.h
@@ -10,7 +10,7 @@ struct FDMXProtocolArtNetPacket
 	: public IDMXProtocolPacket
 {
 public:
-	uint8 ID[ARTNET_STRING_SIZE] = "Art-Net";
+	uint8 ID[ARTNET_STRING_SIZE] = DMX_PROTOCOLNAME_ARTNET;
 };
 
 struct FDMXProtocolArtNetDMXPacket
@@ -56,12 +56,12 @@ public:
 	uint8 Priority = 0;
 };
 
-struct FArtNetPacketReply
+struct FDMXProtocolArtNetPacketReply
 	: public FDMXProtocolArtNetPacket
 {
 public:
 	void Serialize(FArchive& Ar);
-	friend FArchive& operator<<(FArchive& Ar, FArtNetPacketReply& Packet);
+	friend FArchive& operator<<(FArchive& Ar, FDMXProtocolArtNetPacketReply& Packet);
 
 public:
 	uint16 OpCode = 0;
diff --git a/Engine/Plugins/Experimental/DigitalMultiplex/DMXProtocol/Source/DMXProtocolArtNet/Private/Tests/DMXPrtotocolArtNetFlowTest.cpp b/Engine/Plugins/Experimental/DigitalMultiplex/DMXProtocol/Source/DMXProtocolArtNet/Private/Tests/DMXPrtotocolArtNetFlowTest.cpp
--- a/Engine/Plugins/Experimental/DigitalMultiplex/DMXProtocol/Source/DMXProtocolArtNet/Private/Tests/DMXPrtotocolArtNetFlowTest.cpp
+++ b/Engine/Plugins/Experimental/DigitalMultiplex/DMXProtocol/Source/DMXProtocolArtNet/Private/Tests/DMXPrtotocolArtNetFlowTest.cpp
@@ -6,75 +6,30 @@
 #include "DMXProtocolArtNetModule.h"
 
 #include "DMXProtocolArtNet.h"
-#include "DMXProtocolPortArtNet.h"
-
-#include "Managers/DMXProtocolDeviceManager.h"
-#include "Managers/DMXProtocolInterfaceManager.h"
-#include "Managers/DMXProtocolPortManager.h"
 #include "Managers/DMXProtocolUniverseManager.h"
+#include "DMXProtocolUniverseArtNet.h"
 
-#include "DMXProtocolInterfaceArtNet.h"
-#include "DMXProtocolDeviceArtNet.h"
 #include "DMXProtocolUniverseArtNet.h"
-#include "DMXProtocolPortArtNet.h"
 #include "DMXProtocolArtNetConstants.h"
 
-struct FDMXProtocolLauncher
-{
-	FDMXProtocolLauncher(uint8 InDeviceID, uint8 InPortID)
-		: DeviceID(InDeviceID)
-		, PortID(InPortID)
-		, InterfaceID(0)
-	{
-		// Init Protocol and Managers
-		DMXProtocol = IDMXProtocol::Get<FDMXProtocolArtNet>(FDMXProtocolArtNetModule::NAME_Artnet);
-
-		InterfaceArtNet = MakeShared<FDMXProtocolInterfaceArtNet>(DMXProtocol);
-		DMXProtocol->GetInterfaceManager()->AddInterface(InterfaceID, InterfaceArtNet);
-
-		FJsonObject DeviceSettings;
-		ProtocolDevice = MakeShared<FDMXProtocolDeviceArtNet>(DMXProtocol, InterfaceArtNet, DeviceSettings, DeviceID);
-		DMXProtocol->GetDeviceManager()->AddDevice(ProtocolDevice);
-
-		FJsonObject OutputPortSettings;
-		OutputPortSettings.SetNumberField(TEXT("Net"), 0);
-		OutputPortSettings.SetNumberField(TEXT("Subnet"), 1);
-		OutputPortSettings.SetNumberField(TEXT("Universe"), 0);
-		ProtocolOutputPort = MakeShared<FDMXProtocolPortArtNet>(DMXProtocol, ProtocolDevice, OutputPortSettings, PortID, DMX_PORT_OUTPUT);
-		DMXProtocol->GetPortManager()->AddOutputPort(ProtocolDevice, ProtocolOutputPort);
-
-		ProtocolUniverse = MakeShared<FDMXProtocolUniverseArtNet>(DMXProtocol, ProtocolOutputPort);
-		DMXProtocol->GetUniverseManager()->AddUniverse(ProtocolOutputPort, ProtocolUniverse);
-		ProtocolOutputPort->SetUniverse(ProtocolUniverse);
-	}
-
-	~FDMXProtocolLauncher()
-	{
-		DMXProtocol->GetInterfaceManager()->RemoveInterface(InterfaceID);
-		DMXProtocol->GetDeviceManager()->RemoveDevice(ProtocolDevice);
-		DMXProtocol->GetPortManager()->RemoveOutputPort(ProtocolDevice, ProtocolOutputPort);
-		DMXProtocol->GetUniverseManager()->RemoveUniverse(ProtocolOutputPort, ProtocolUniverse);
-	}
-
-	const uint8 DeviceID;
-	const uint8 PortID;
-	const uint8 InterfaceID;
-	TSharedPtr<FDMXProtocolInterfaceArtNet> InterfaceArtNet;
-	TSharedPtr<FDMXProtocolDeviceArtNet> ProtocolDevice;
-	TSharedPtr<FDMXProtocolPortArtNet> ProtocolOutputPort;
-	TSharedPtr<FDMXProtocolUniverseArtNet> ProtocolUniverse;
-	FDMXProtocolArtNet* DMXProtocol;
-};
 
 IMPLEMENT_SIMPLE_AUTOMATION_TEST(FDMXPrtotocolArtNetBasicFlowTest, "VirtualProduction.DMX.ArtNet.BasicFlow", EAutomationTestFlags::ApplicationContextMask | EAutomationTestFlags::EngineFilter)
 
-static void SensDMXFragment(FAutomationTestBase& Test, const FDMXProtocolLauncher& DMXProtocolLauncher)
+bool FDMXPrtotocolArtNetBasicFlowTest::RunTest(const FString& Parameters)
 {
-	FDMXProtocolPortArtNet* ArtNetPort = static_cast<FDMXProtocolPortArtNet*>(DMXProtocolLauncher.DMXProtocol->GetPortManager()->GetPortByDeviceAndID(DMXProtocolLauncher.ProtocolDevice, DMXProtocolLauncher.PortID, EDMXPortDirection::DMX_PORT_OUTPUT));
+	static const uint16 UniverseValue = 0;
+
+	// Init Protocol and Managers
+	FDMXProtocolArtNet* DMXProtocol = static_cast<FDMXProtocolArtNet*>(IDMXProtocol::Get(FDMXProtocolArtNetModule::NAME_Artnet).Get());
 
-	Test.TestNotNull(TEXT("ArtNetPort is valid"), ArtNetPort);
+	FJsonObject UniverseSettings;
+	UniverseSettings.SetNumberField(TEXT("UniverseID"), UniverseValue);
+	UniverseSettings.SetNumberField(TEXT("PortID"), 0); // TODO set PortID
+	DMXProtocol->AddUniverse(UniverseSettings);
 
-	Test.TestEqual(TEXT("Universe ID should be equal 16"), ArtNetPort->GetUniverseID(), 16);
+	TSharedPtr<IDMXProtocolUniverse, ESPMode::ThreadSafe> Universe = DMXProtocol->GetUniverseById(UniverseValue);
+	TestTrue(TEXT("Universe is valid"), Universe.IsValid());
+	FDMXProtocolUniverseArtNet* UniverseArtNet = static_cast<FDMXProtocolUniverseArtNet*>(Universe.Get());
 
 	// Send DMX fragment 1
 	const uint8 FixtureChannels[6] = { 30, 31, 32, 36, 37, 38 };
@@ -83,91 +38,70 @@ static void SensDMXFragment(FAutomationTestBase& Test, const FDMXProtocolLaunche
 	DMXFragmentMap.Add(FixtureChannels[0], FixtureValues[0]);
 	DMXFragmentMap.Add(FixtureChannels[1], FixtureValues[1]);
 	DMXFragmentMap.Add(FixtureChannels[2], FixtureValues[2]);
-	DMXProtocolLauncher.DMXProtocol->SetDMXFragment(ArtNetPort->GetUniverseID(), DMXFragmentMap);
+	DMXProtocol->SendDMXFragment(UniverseValue, DMXFragmentMap);
 
 	// Wait before recieve DMX
 	FPlatformProcess::Sleep(0.2f);
-	TSharedPtr<IDMXProtocolUniverse> Universe = ArtNetPort->GetCachedUniverse().Pin();
-	Test.TestTrue(TEXT("Universe should be valid"), Universe.IsValid());
-	TSharedPtr<FDMXBuffer> InputDMXBuffer = Universe->GetInputDMXBuffer();
+	TSharedPtr<FDMXBuffer> InputDMXBuffer = UniverseArtNet->GetInputDMXBuffer();
 
-	Test.TestEqual(TEXT("Incoming buffer should be same"), FixtureValues[0], InputDMXBuffer->GetDMXData()[FixtureChannels[0]]);
-	Test.TestEqual(TEXT("Incoming buffer should be same"), FixtureValues[1], InputDMXBuffer->GetDMXData()[FixtureChannels[1]]);
-	Test.TestEqual(TEXT("Incoming buffer should be same"), FixtureValues[2], InputDMXBuffer->GetDMXData()[FixtureChannels[2]]);
+	TestEqual(TEXT("Incoming buffer should be same"), FixtureValues[0], InputDMXBuffer->GetDMXData()[FixtureChannels[0]]);
+	TestEqual(TEXT("Incoming buffer should be same"), FixtureValues[1], InputDMXBuffer->GetDMXData()[FixtureChannels[1]]);
+	TestEqual(TEXT("Incoming buffer should be same"), FixtureValues[2], InputDMXBuffer->GetDMXData()[FixtureChannels[2]]);
 
 	// Send DMX fragment 2
 	DMXFragmentMap.Add(FixtureChannels[3], FixtureValues[3]);
 	DMXFragmentMap.Add(FixtureChannels[4], FixtureValues[4]);
 	DMXFragmentMap.Add(FixtureChannels[5], FixtureValues[5]);
-	DMXProtocolLauncher.DMXProtocol->SetDMXFragment(ArtNetPort->GetUniverseID(), DMXFragmentMap);
+	DMXProtocol->SendDMXFragment(UniverseValue, DMXFragmentMap);
 
 	FPlatformProcess::Sleep(0.2f);
 
-	Test.TestEqual(TEXT("Old value DMX input should be the same"), FixtureValues[0], InputDMXBuffer->GetDMXData()[FixtureChannels[0]]);
-	Test.TestEqual(TEXT("Incoming buffer should be same"), FixtureValues[4], InputDMXBuffer->GetDMXData()[FixtureChannels[4]]);
-	Test.TestEqual(TEXT("Incoming buffer should be same"), FixtureValues[5], InputDMXBuffer->GetDMXData()[FixtureChannels[5]]);
+	TestEqual(TEXT("Old value DMX input should be the same"), FixtureValues[0], InputDMXBuffer->GetDMXData()[FixtureChannels[0]]);
+	TestEqual(TEXT("Incoming buffer should be same"), FixtureValues[4], InputDMXBuffer->GetDMXData()[FixtureChannels[4]]);
+	TestEqual(TEXT("Incoming buffer should be same"), FixtureValues[5], InputDMXBuffer->GetDMXData()[FixtureChannels[5]]);
 
 	// @TODO Implement test cases for Pool and RDM
 	// Send DMX pull
-	DMXProtocolLauncher.DMXProtocol->TransmitPoll();
+	DMXProtocol->TransmitPoll();
+	FPlatformProcess::Sleep(0.2f);
+	const FDMXProtocolArtNetPollPacket& PollPacket = DMXProtocol->GetIncomingPoll();
+	TestEqual(TEXT("OpCode should be ARTNET_POLL"), ARTNET_POLL, PollPacket.OpCode);
 
-	// Single TOD Request
-	{
-		// Send TOD to non existent device
-		{
-			FJsonObject DeviceSettings;
-			TSharedPtr<FDMXProtocolDeviceArtNet> NonExistentDevice = MakeShared<FDMXProtocolDeviceArtNet>(DMXProtocolLauncher.DMXProtocol, DMXProtocolLauncher.InterfaceArtNet, DeviceSettings, 987);
-			DMXProtocolLauncher.DMXProtocol->TransmitTodRequest(NonExistentDevice);
-			FPlatformProcess::Sleep(0.2f);
-			const FDMXProtocolArtNetTodRequest& IncomingTodRequest = DMXProtocolLauncher.DMXProtocol->GetIncomingTodRequest();
-			Test.TestEqual(TEXT("TodRequest.AdCount should be 0"), IncomingTodRequest.AdCount, uint8(0));
-		}
-
-		// Send TOD to existing device
-		{
-			DMXProtocolLauncher.DMXProtocol->TransmitTodRequest(DMXProtocolLauncher.ProtocolDevice);
-			FPlatformProcess::Sleep(0.2f);
-			const FDMXProtocolArtNetTodRequest& IncomingTodRequest = DMXProtocolLauncher.DMXProtocol->GetIncomingTodRequest();
-			Test.TestNotEqual(TEXT("TodRequest.AdCount should not be 0"), IncomingTodRequest.AdCount, uint8(0));
-		}
-	}
+	// Send TOD 
+	DMXProtocol->TransmitTodRequest(UniverseArtNet->GetPortAddress());
+	FPlatformProcess::Sleep(0.2f);
+	const FDMXProtocolArtNetTodRequest& IncomingTodRequest = DMXProtocol->GetIncomingTodRequest();
+	TestEqual(TEXT("TodRequest.AdCount should not be 1"), IncomingTodRequest.AdCount, uint8(1));
 
 	// Test ArtNet RDM
-	DMXProtocolLauncher.DMXProtocol->TransmitTodRequestToAll();
-	ArtNetPort->AddTODUID(FRDMUID({ 0x00, 0x01, 0x02, 0x03, 0x04, 0x05 }));
-	ArtNetPort->AddTODUID(FRDMUID({ 0x06, 0x07, 0x08, 0x09, 0x10, 0x11 }));
+	DMXProtocol->TransmitTodRequestToAll();
+	UniverseArtNet->AddTODUID(FRDMUID({ 0x00, 0x01, 0x02, 0x03, 0x04, 0x05 }));
+	UniverseArtNet->AddTODUID(FRDMUID({ 0x06, 0x07, 0x08, 0x09, 0x10, 0x11 }));
 
 	FRDMUID RDMUID3;
 	uint8 RDMBuffer3[RDM_UID_WIDTH] = { 0x12, 0x13, 0x14, 0x15, 0x16, 0x17 };
 	FMemory::Memcpy(RDMUID3.Buffer, RDMBuffer3, RDM_UID_WIDTH);
-	ArtNetPort->AddTODUID(RDMUID3);
+	UniverseArtNet->AddTODUID(RDMUID3);
 
 	// Test wrong size of RDM Buffer C array
 	uint8 RDMBuffer4[3] = { 0x18, 0x19, 0x20 };
-	ArtNetPort->AddTODUID(FRDMUID(RDMBuffer4));
+	UniverseArtNet->AddTODUID(FRDMUID(RDMBuffer4));
 
 	// Test wrong size of RDM Buffer C array
 	uint8 RDMBuffer5[10] = { 0x26, 0x27, 0x28, 0x29, 0x30, 0x31, 0x32, 0x32, 0x33 };
-	ArtNetPort->AddTODUID(FRDMUID(RDMBuffer5));
+	UniverseArtNet->AddTODUID(FRDMUID(RDMBuffer5));
 
 	// Test wrong size of RDM Buffer TArray
-	ArtNetPort->AddTODUID(FRDMUID({ 0x34, 0x35 }));
+	UniverseArtNet->AddTODUID(FRDMUID({ 0x34, 0x35 }));
 
-	DMXProtocolLauncher.DMXProtocol->TransmitTodData(DMXProtocolLauncher.ProtocolDevice, DMXProtocolLauncher.PortID);
-	DMXProtocolLauncher.DMXProtocol->TransmitTodControl(DMXProtocolLauncher.ProtocolDevice, DMXProtocolLauncher.PortID, 0x02);
+	DMXProtocol->TransmitTodData(UniverseArtNet->GetUniverseID());
+	DMXProtocol->TransmitTodControl(UniverseArtNet->GetUniverseID(), 0x02);
 
 	TArray<uint8> RDMData;
 	RDMData.AddZeroed(ARTNET_MAX_RDM_DATA);
 	RDMData[47] = 0x00;
 	RDMData[48] = 0xcc;
-	DMXProtocolLauncher.DMXProtocol->TransmitRDM(DMXProtocolLauncher.ProtocolDevice, DMXProtocolLauncher.PortID, RDMData);
-}
-
-bool FDMXPrtotocolArtNetBasicFlowTest::RunTest(const FString& Parameters)
-{
-	FDMXProtocolLauncher DMXProtocolLauncher(0, 0);
-	// Test Sending DMX and RDM
-	SensDMXFragment(*this, DMXProtocolLauncher);
+	DMXProtocol->TransmitRDM(UniverseArtNet->GetUniverseID(), RDMData);
 
 	return true;
 }
@@ -176,7 +110,16 @@ IMPLEMENT_SIMPLE_AUTOMATION_TEST(FDMXPrtotocolArtNetConsoleCommandsTest, "Virtua
 
 bool FDMXPrtotocolArtNetConsoleCommandsTest::RunTest(const FString& Parameters)
 {
-	FDMXProtocolLauncher DMXProtocolLauncher(0, 0);
+	static const uint16 UniverseValue = 1891;
+
+	FDMXProtocolArtNet* DMXProtocol = static_cast<FDMXProtocolArtNet*>(IDMXProtocol::Get(FDMXProtocolArtNetModule::NAME_Artnet).Get());
+
+	FJsonObject UniverseSettings;
+	UniverseSettings.SetNumberField(TEXT("UniverseID"), UniverseValue);
+	UniverseSettings.SetNumberField(TEXT("PortID"), 0); // TODO set PortID
+	DMXProtocol->AddUniverse(UniverseSettings);
+	TSharedPtr<IDMXProtocolUniverse, ESPMode::ThreadSafe> Universe = DMXProtocol->GetUniverseById(UniverseValue);
+
 
 	// DMX.ArtNet.SendDMX[UniverseID] Channel:Value Channel : Value Channel : Value ...
 	uint32 Values[3][2] =
@@ -188,7 +131,7 @@ bool FDMXPrtotocolArtNetConsoleCommandsTest::RunTest(const FString& Parameters)
 
 	FString Command(
 		FString::Printf(TEXT("DMX.ArtNet.SendDMX %d %d:%d %d:%d %d:%d"),
-			DMXProtocolLauncher.ProtocolOutputPort->GetUniverseID(),
+			UniverseValue,
 			Values[0][0], Values[0][1],
 			Values[1][0], Values[1][1],
 			Values[2][0], Values[2][1]
@@ -198,7 +141,6 @@ bool FDMXPrtotocolArtNetConsoleCommandsTest::RunTest(const FString& Parameters)
 
 	FPlatformProcess::Sleep(0.2f);
 
-	TSharedPtr<IDMXProtocolUniverse> Universe = DMXProtocolLauncher.ProtocolOutputPort->GetCachedUniverse().Pin();
 	TSharedPtr<FDMXBuffer> InputDMXBuffer = Universe->GetInputDMXBuffer();
 	TestEqual(TEXT("Incoming buffer should be same"), Values[0][1], InputDMXBuffer->GetDMXData()[Values[0][0]]);
 	TestEqual(TEXT("Incoming buffer should be same"), Values[1][1], InputDMXBuffer->GetDMXData()[Values[1][0]]);
diff --git a/Engine/Plugins/Experimental/DigitalMultiplex/DMXProtocol/Source/DMXProtocolSACN/Private/DMXProtocolSACN.cpp b/Engine/Plugins/Experimental/DigitalMultiplex/DMXProtocol/Source/DMXProtocolSACN/Private/DMXProtocolSACN.cpp
--- a/Engine/Plugins/Experimental/DigitalMultiplex/DMXProtocol/Source/DMXProtocolSACN/Private/DMXProtocolSACN.cpp
+++ b/Engine/Plugins/Experimental/DigitalMultiplex/DMXProtocol/Source/DMXProtocolSACN/Private/DMXProtocolSACN.cpp
@@ -7,52 +7,49 @@
 #include "Packets/DMXProtocolE131PDUPacket.h"
 #include "DMXProtocolPackager.h"
 #include "DMXProtocolSettings.h"
+#include "DMXProtocolUniverseSACN.h"
+#include "Managers/DMXProtocolUniverseManager.h"
 
 #include "Serialization/BufferArchive.h"
 
-TSharedPtr<IDMXProtocolSender> FDMXProtocolSACN::GetSenderInterface() const
+FDMXProtocolSACN::FDMXProtocolSACN(const FName& InProtocolName, FJsonObject& InSettings)
+	: ProtocolName(InProtocolName)
+	, SenderSocket(nullptr)
 {
-	return SACNSender;
+	NetworkErrorMessagePrefix = TEXT("NETWORK ERROR SACN:");
+
+	Settings = MakeShared<FJsonObject>(InSettings);
+	UniverseManager = MakeShared<FDMXProtocolUniverseManager<FDMXProtocolUniverseSACN>>(this);
+}
+
+
+TSharedPtr<FJsonObject> FDMXProtocolSACN::GetSettings() const
+{
+	return Settings;
 }
 
 bool FDMXProtocolSACN::Init()
 {	
-	//Socket;
-	ISocketSubsystem* SocketSubsystem = ISocketSubsystem::Get(PLATFORM_SOCKETSUBSYSTEM);
-	bool bIsValid = false;
-
-	// Create broadcast address
-	FIPv4Endpoint UnicastEndpoint = FIPv4Endpoint::Any;
+	InterfaceIPAddress = GetDefault<UDMXProtocolSettings>()->InterfaceIPAddress;
 
-	// Create sender socket
-	SenderSocket = FUdpSocketBuilder(TEXT("UDPSACNSenderSocket"))
-		.AsNonBlocking()
-#if PLATFORM_WINDOWS
-		.BoundToAddress(UnicastEndpoint.Address)
-#endif
-		.AsReusable();
+	// Set Network Interface listener
+	NetworkInterfaceChangedDelegate = FOnNetworkInterfaceChangedDelegate::CreateRaw(this, &FDMXProtocolSACN::OnNetworkInterfaceChanged);
+	IDMXProtocol::OnNetworkInterfaceChanged.Add(NetworkInterfaceChangedDelegate);
 
-	if (SenderSocket == nullptr)
+	// Set Network Interface
+	FString ErrorMessage;
+	if (!RestartNetworkInterface(InterfaceIPAddress, ErrorMessage))
 	{
-		UE_LOG_DMXPROTOCOL(Error, TEXT("ERROR create SenderSocket"));
-		return false;
+		UE_LOG_DMXPROTOCOL(Error, TEXT("%s:%s"), NetworkErrorMessagePrefix, *ErrorMessage);
 	}
 
-	SACNSender = MakeShared<FDMXProtocolSenderSACN>(*SenderSocket, this);
-
 	return true;
 }
 
 bool FDMXProtocolSACN::Shutdown()
 {
-	SACNSender.Reset();
-
-	//Clear all sockets!
-	if (SenderSocket)
-	{
-		SenderSocket->Close();
-		ISocketSubsystem::Get(PLATFORM_SOCKETSUBSYSTEM)->DestroySocket(SenderSocket);
-	}
+	ReleaseNetworkInterface();
+	IDMXProtocol::OnNetworkInterfaceChanged.Remove(NetworkInterfaceChangedDelegate.GetHandle());
 
 	return false;
 }
@@ -62,13 +59,52 @@ bool FDMXProtocolSACN::IsEnabled() const
 	return true;
 }
 
-FName FDMXProtocolSACN::GetProtocolName() const
+TSharedPtr<IDMXProtocolUniverse, ESPMode::ThreadSafe> FDMXProtocolSACN::AddUniverse(const FJsonObject& InSettings)
+{
+	TSharedPtr<FDMXProtocolUniverseSACN, ESPMode::ThreadSafe> Universe = MakeShared<FDMXProtocolUniverseSACN, ESPMode::ThreadSafe>(SharedThis(this), InSettings);
+	return UniverseManager->AddUniverse(Universe->GetUniverseID(), Universe);
+}
+
+void FDMXProtocolSACN::CollectUniverses(const TArray<FDMXUniverse>& Universes)
+{
+	for (const FDMXUniverse& Universe : Universes)
+	{
+		if (UniverseManager->GetAllUniverses().Contains(Universe.UniverseNumber))
+		{
+			continue;
+		}
+
+		FJsonObject UniverseSettings;
+		UniverseSettings.SetNumberField(TEXT("UniverseID"), Universe.UniverseNumber);
+		AddUniverse(UniverseSettings);
+	}
+}
+
+bool FDMXProtocolSACN::RemoveUniverseById(uint32 InUniverseId)
+{
+	return UniverseManager->RemoveUniverseById(InUniverseId);
+}
+
+void FDMXProtocolSACN::RemoveAllUniverses()
+{
+	UniverseManager->RemoveAll();
+}
+
+TSharedPtr<IDMXProtocolUniverse, ESPMode::ThreadSafe> FDMXProtocolSACN::GetUniverseById(uint32 InUniverseId) const
+{
+	return UniverseManager->GetUniverseById(InUniverseId);
+}
+
+
+const FName& FDMXProtocolSACN::GetProtocolName() const
 {
-	return FDMXProtocolSACNModule::NAME_SACN;
+	return ProtocolName;
 }
 
-void FDMXProtocolSACN::Reload()
+
+TSharedPtr<IDMXProtocolSender> FDMXProtocolSACN::GetSenderInterface() const
 {
+	return SACNSender;
 }
 
 bool FDMXProtocolSACN::Tick(float DeltaTime)
@@ -82,9 +118,7 @@ bool FDMXProtocolSACN::SendDiscovery(const TArray<uint16>& Universes)
 	FDMXProtocolPackager Packager;
 
 	FDMXProtocolE131RootLayerPacket RootLayer;
-	RootLayer.Flags = 0x72; // Const for now
-	RootLayer.Length = 0x6e; // Const for now
-	FGuid Guid = FGuid::NewGuid();
+	static FGuid Guid = FGuid::NewGuid();
 	FMemory::Memcpy(RootLayer.CID, &Guid, ACN_CIDBYTES);
 
 	FDMXProtocolUDPE131FramingLayerPacket FramingLayer;
@@ -95,7 +129,7 @@ bool FDMXProtocolSACN::SendDiscovery(const TArray<uint16>& Universes)
 	}
 	else
 	{
-		UE_LOG_DMXPROTOCOL(Warning, TEXT("Size of outcoming sACN DMX discovery buffer is wrong"));
+		UE_LOG_DMXPROTOCOL(Error, TEXT("%s:Size of outcoming sACN DMX discovery buffer is wrong"), NetworkErrorMessagePrefix);
 	}
 
 	Packager.AddToPackage(&RootLayer);
@@ -111,46 +145,88 @@ bool FDMXProtocolSACN::SendDiscovery(const TArray<uint16>& Universes)
 	return true;
 }
 
-bool FDMXProtocolSACN::SendDMX(uint16 UniverseID, uint8 PortID, const TSharedPtr<FDMXBuffer>& DMXBuffer) const
+EDMXSendResult FDMXProtocolSACN::SendDMXFragment(uint16 InUniverseID, const IDMXFragmentMap& DMXFragment)
+{
+	uint16 FinalSendUniverseID = GetFinalSendUniverseID(InUniverseID);
+
+	TSharedPtr<IDMXProtocolUniverse, ESPMode::ThreadSafe> Universe = UniverseManager->GetUniverseById(FinalSendUniverseID);
+	if (!Universe.IsValid())
+	{
+		return EDMXSendResult::ErrorGetUniverse;
+	}
+
+	if (!Universe->SetDMXFragment(DMXFragment))
+	{
+		return EDMXSendResult::ErrorSetBuffer;
+	}
+
+	EDMXSendResult Result = SendDMXInternal(FinalSendUniverseID, Universe->GetOutputDMXBuffer());
+	return Result;
+}
+
+EDMXSendResult FDMXProtocolSACN::SendDMXFragmentCreate(uint16 InUniverseID, const IDMXFragmentMap& DMXFragment)
+{
+	uint16 FinalSendUniverseID = GetFinalSendUniverseID(InUniverseID);
+
+	TSharedPtr<IDMXProtocolUniverse, ESPMode::ThreadSafe> Universe = UniverseManager->AddUniverseCreate(FinalSendUniverseID);
+
+	if (!Universe->SetDMXFragment(DMXFragment))
+	{
+		return EDMXSendResult::ErrorSetBuffer;
+	}
+
+	EDMXSendResult Result = SendDMXInternal(FinalSendUniverseID, Universe->GetOutputDMXBuffer());
+	return Result;
+}
+
+uint16 FDMXProtocolSACN::GetFinalSendUniverseID(uint16 InUniverseID) const
+{
+	// GlobalSACNUniverseOffset clamp between 0 and 65535(max uint16) 
+	return InUniverseID + GetDefault<UDMXProtocolSettings>()->GlobalArtNetUniverseOffset;
+}
+
+uint32 FDMXProtocolSACN::GetUniversesNum() const
+{
+	return UniverseManager->GetAllUniverses().Num();
+}
+
+uint16 FDMXProtocolSACN::GetMinUniverseID() const
+{
+	return 1;
+}
+
+uint16 FDMXProtocolSACN::GetMaxUniverses() const
+{
+	return ACN_MAX_UNIVERSES;
+}
+
+EDMXSendResult FDMXProtocolSACN::SendDMXInternal(uint16 InUniverseID, const TSharedPtr<FDMXBuffer>& DMXBuffer) const
 {
 	// Init Packager
 	FDMXProtocolPackager Packager;
 
 	// SACN PDU packets
 	FDMXProtocolE131RootLayerPacket RootLayer;;
-	RootLayer.PreambleSize = ACN_RLP_PREAMBLE_SIZE;
-	RootLayer.PostambleSize = 0x0000;
-	RootLayer.Flags = 0x72; // Const for now
-	RootLayer.Length = 0x6e; // Const for now
-	RootLayer.Vector = 4; // Const for now
-	FGuid Guid = FGuid::NewGuid();
+	static FGuid Guid = FGuid::NewGuid();
 	FMemory::Memcpy(RootLayer.CID, &Guid, ACN_CIDBYTES);
 
 	FDMXProtocolE131FramingLayerPacket FramingLayer;
-	FramingLayer.Flags = 0x72; // Const for now
-	FramingLayer.Length = 0x58; // Const for now
-	FramingLayer.Priority = 100; // Const for now
-	FramingLayer.SynchronizationAddress = 0; // Const for now
-	FramingLayer.SequenceNumber = 0; // Const for now
-	FramingLayer.Options = 0; // Const for now
-	FramingLayer.Universe = UniverseID;
+	// To detect duplicate or out of order packets. This is using on physical controllers
+	static uint8 SequenceNumber = 0;
+	FramingLayer.SequenceNumber = SequenceNumber++;
+	FramingLayer.Universe = InUniverseID;
 
 	FDMXProtocolE131DMPLayerPacket DMPLayer;
+	DMPLayer.AddressIncrement = ACN_ADDRESS_INC;
+	DMPLayer.PropertyValueCount = ACN_DMX_SIZE + 1;
 	if (DMXBuffer->GetDMXData().Num() == ACN_DMX_SIZE)
 	{
 		FMemory::Memcpy(DMPLayer.DMX, DMXBuffer->GetDMXData().GetData(), ACN_DMX_SIZE);
 	}
 	else
 	{
-		UE_LOG_DMXPROTOCOL(Warning, TEXT("Size of outcoming sACN DMX buffer is wrong"));
+		return EDMXSendResult::ErrorSizeBuffer;
 	}
-	DMPLayer.Flags = 0x72; // Const for now
-	DMPLayer.Length = 0x0b; // Const for now
-	DMPLayer.AddressTypeAndDataType = 0xa1; // Const for now
-	DMPLayer.FirstPropertyAddress = 0; // Const for now
-	DMPLayer.AddressIncrement = ACN_ADDRESS_INC;
-	DMPLayer.PropertyValueCount = ACN_DMX_SIZE + 1;
-	DMPLayer.STARTCode = 0; // Const for now
 
 	Packager.AddToPackage(&RootLayer);
 	Packager.AddToPackage(&FramingLayer);
@@ -158,11 +234,94 @@ bool FDMXProtocolSACN::SendDMX(uint16 UniverseID, uint8 PortID, const TSharedPtr
 
 	// Sending
 	FJsonObject PacketSettings;
-	PacketSettings.SetNumberField("UniverseID", UniverseID);
+	PacketSettings.SetNumberField("UniverseID", InUniverseID);
 	FDMXPacketPtr Packet = MakeShared<FDMXPacket, ESPMode::ThreadSafe>(PacketSettings, Packager.GetBuffer());
-	GetSenderInterface()->EnqueueOutboundPackage(Packet);
 
-	return false;
+	if (!GetSenderInterface()->EnqueueOutboundPackage(Packet))
+	{
+		return EDMXSendResult::ErrorEnqueuePackage;
+	}
+
+	return EDMXSendResult::Success;
+}
+
+void FDMXProtocolSACN::OnNetworkInterfaceChanged(const FString& InInterfaceIPAddress)
+{
+	FString ErrorMessage;
+	if (!RestartNetworkInterface(InInterfaceIPAddress, ErrorMessage))
+	{
+		UE_LOG_DMXPROTOCOL(Error, TEXT("%s:%s"), NetworkErrorMessagePrefix, *ErrorMessage);
+	}
+}
+
+bool FDMXProtocolSACN::RestartNetworkInterface(const FString& InInterfaceIPAddress, FString& OutErrorMessage)
+{
+	FScopeLock Lock(&SenderSocketCS);
+
+	// Clean the error message
+	OutErrorMessage.Empty();
+
+	// Try to create IP address at the first
+	ISocketSubsystem* SocketSubsystem = ISocketSubsystem::Get(PLATFORM_SOCKETSUBSYSTEM);
+	// Create sender address
+	TSharedPtr<FInternetAddr> SenderAddr = SocketSubsystem->CreateInternetAddr();
+	bool bIsValid = false;
+	SenderAddr->SetIp(*InInterfaceIPAddress, bIsValid);
+	SenderAddr->SetPort(ACN_PORT);
+	if (!bIsValid)
+	{
+		OutErrorMessage = FString::Printf(TEXT("Wrong IP address: %s"), *InInterfaceIPAddress);
+		return false;
+	}
+
+	// Release old network interface
+	ReleaseNetworkInterface();
+
+	FIPv4Endpoint SenderEndpoint = FIPv4Endpoint(SenderAddr);
+
+	// Create sender socket
+	FSocket* NewSenderSocket = FUdpSocketBuilder(TEXT("UDPSACNSenderSocket"))
+		.AsNonBlocking()
+#if PLATFORM_SUPPORTS_UDP_MULTICAST_GROUP
+        .WithMulticastInterface(SenderEndpoint.Address)
+        .WithMulticastLoopback()
+        .WithMulticastTtl(1)
+#endif
+        .BoundToEndpoint(SenderEndpoint)
+		.AsReusable();
+
+	if (NewSenderSocket == nullptr)
+	{
+		OutErrorMessage = FString::Printf(TEXT("Error create SenderSocket: %s"), *InterfaceIPAddress);
+		return false;
+	}
+
+	// Set new network interface IP
+	InterfaceIPAddress = InInterfaceIPAddress;
+
+	// Save New socket;
+	SenderSocket = NewSenderSocket;
+
+	SACNSender = MakeShared<FDMXProtocolSenderSACN>(*SenderSocket, this);
+
+	return true;
+}
+
+void FDMXProtocolSACN::ReleaseNetworkInterface()
+{
+	if (SACNSender)
+	{
+		SACNSender.Reset();
+		SACNSender = nullptr;
+	}
+
+	// Clean all sockets!
+	if (SenderSocket != nullptr)
+	{
+		SenderSocket->Close();
+		ISocketSubsystem::Get(PLATFORM_SOCKETSUBSYSTEM)->DestroySocket(SenderSocket);
+		SenderSocket = nullptr;
+	}
 }
 
 void FDMXProtocolSACN::SendRDMCommand(const TSharedPtr<FJsonObject>& CMD)
@@ -187,7 +346,7 @@ TSharedPtr<FInternetAddr> FDMXProtocolSACN::GetUniverseAddr(uint16 InUniverseID)
 	uint8 IP_1 = ACN_UNIVERSE_IP_1;
 	IP << IP_0; // [x.?.?.?]
 	IP << IP_1; // [x.x.?.?]
-	IP.ByteSwap(&InUniverseID, 2);
+	IP.ByteSwap(&InUniverseID, sizeof(uint16));
 	IP << InUniverseID;	// [x.x.x.x]
 
 	InternetAddr->SetRawIp(IP);
diff --git a/Engine/Plugins/Experimental/DigitalMultiplex/DMXProtocol/Source/DMXProtocolSACN/Private/DMXProtocolSACN.h b/Engine/Plugins/Experimental/DigitalMultiplex/DMXProtocol/Source/DMXProtocolSACN/Private/DMXProtocolSACN.h
new file mode 100644
--- /dev/null
+++ b/Engine/Plugins/Experimental/DigitalMultiplex/DMXProtocol/Source/DMXProtocolSACN/Private/DMXProtocolSACN.h
@@ -0,0 +1,101 @@
+// Copyright Epic Games, Inc. All Rights Reserved.
+
+#pragma once
+
+#include "CoreMinimal.h"
+#include "Interfaces/IDMXProtocol.h"
+#include "Interfaces/IDMXNetworkInterface.h"
+#include "DMXProtocolSACNModule.h"
+
+#include "Interfaces/IPv4/IPv4Endpoint.h"
+#include "HAL/CriticalSection.h"
+#include "HAL/ThreadSafeCounter.h"
+
+class FInternetAddr;
+class FSocket;
+class FDMXProtocolUniverseSACN;
+
+template<class TUniverse>
+class FDMXProtocolUniverseManager;
+
+class DMXPROTOCOLSACN_API FDMXProtocolSACN
+	: public IDMXProtocol
+	, public IDMXNetworkInterface
+{
+public:
+	explicit FDMXProtocolSACN(const FName& InProtocolName, FJsonObject& InSettings);
+
+	//~ Begin IDMXProtocolBase implementation
+	virtual bool Init() override;
+	virtual bool Shutdown() override;
+	virtual bool Tick(float DeltaTime) override;
+	//~ End IDMXProtocolBase implementation
+
+	//~ Begin IDMXProtocol implementation
+	virtual const FName& GetProtocolName() const override;
+	virtual TSharedPtr<IDMXProtocolSender> GetSenderInterface() const override;
+	virtual TSharedPtr<FJsonObject> GetSettings() const override;
+	virtual bool IsEnabled() const override;
+	virtual TSharedPtr<IDMXProtocolUniverse, ESPMode::ThreadSafe> AddUniverse(const FJsonObject& InSettings) override;
+	virtual void CollectUniverses(const TArray<FDMXUniverse>& Universes) override;
+	virtual bool RemoveUniverseById(uint32 InUniverseId) override;
+	virtual void RemoveAllUniverses() override;
+	virtual TSharedPtr<IDMXProtocolUniverse, ESPMode::ThreadSafe> GetUniverseById(uint32 InUniverseId) const override;
+	virtual EDMXSendResult SendDMXFragment(uint16 InUniverseID, const IDMXFragmentMap& DMXFragment) override;
+	virtual EDMXSendResult SendDMXFragmentCreate(uint16 InUniverseID, const IDMXFragmentMap& DMXFragment) override;
+	virtual uint16 GetFinalSendUniverseID(uint16 InUniverseID) const override;
+	virtual uint32 GetUniversesNum() const override;
+	virtual uint16 GetMinUniverseID() const override;
+	virtual uint16 GetMaxUniverses() const override;
+	virtual FOnUniverseInputUpdateEvent& GetOnUniverseInputUpdate() override
+	{
+		return OnUniverseInputUpdateEvent;
+	}
+	//~ End IDMXProtocol implementation
+
+	//~ Begin IDMXNetworkInterface implementation
+	virtual void OnNetworkInterfaceChanged(const FString& InInterfaceIPAddress) override;
+	virtual bool RestartNetworkInterface(const FString& InInterfaceIPAddress, FString& OutErrorMessage) override;
+	virtual void ReleaseNetworkInterface() override;
+	//~ End IDMXNetworkInterface implementation
+
+	//~ Begin IDMXProtocolRDM implementation
+	virtual void SendRDMCommand(const TSharedPtr<FJsonObject>& CMD) override;
+	virtual void RDMDiscovery(const TSharedPtr<FJsonObject>& CMD) override;
+	//~ End IDMXProtocol implementation
+
+	//~ sACN specific implementation
+	bool SendDiscovery(const TArray<uint16>& Universes);
+
+	//~ Only the factory makes instances
+	FDMXProtocolSACN() = delete;
+
+public:
+	static TSharedPtr<FInternetAddr> GetUniverseAddr(uint16 InUniverseID);
+
+private:
+	EDMXSendResult SendDMXInternal(uint16 UniverseID, const TSharedPtr<FDMXBuffer>& DMXBuffer) const;
+
+private:
+	FName ProtocolName;
+
+	TSharedPtr<FJsonObject> Settings;
+
+	TSharedPtr<FDMXProtocolUniverseManager<FDMXProtocolUniverseSACN>> UniverseManager;
+
+	TSharedPtr<IDMXProtocolSender> SACNSender;
+
+	/** Holds the network socket used to sender packages. */
+	FSocket* SenderSocket;
+
+	FOnUniverseInputUpdateEvent OnUniverseInputUpdateEvent;
+
+	/** Mutex protecting access to the listening socket. */
+	mutable FCriticalSection SenderSocketCS;
+
+	FString InterfaceIPAddress;
+
+	FOnNetworkInterfaceChangedDelegate NetworkInterfaceChangedDelegate;
+
+	const TCHAR* NetworkErrorMessagePrefix;
+};
diff --git a/Engine/Plugins/Experimental/DigitalMultiplex/DMXProtocol/Source/DMXProtocolSACN/Private/DMXProtocolSACNConstants.h b/Engine/Plugins/Experimental/DigitalMultiplex/DMXProtocol/Source/DMXProtocolSACN/Private/DMXProtocolSACNConstants.h
--- a/Engine/Plugins/Experimental/DigitalMultiplex/DMXProtocol/Source/DMXProtocolSACN/Private/DMXProtocolSACNConstants.h
+++ b/Engine/Plugins/Experimental/DigitalMultiplex/DMXProtocol/Source/DMXProtocolSACN/Private/DMXProtocolSACNConstants.h
@@ -5,6 +5,8 @@
 #include "CoreMinimal.h"
 #include "DMXProtocolTypes.h"
 
+#define DMX_PROTOCOLNAME_SACN "sACN"
+
 enum 
 {
 	VECTOR_ROOT_E131_DATA = 0x00000004,
@@ -36,7 +38,12 @@ enum
 	ACN_CIDBYTES = 16,
 	ACN_SOURCE_NAME_SIZE = 64,
 	ACN_ADDRESS_INC = 1,
-	ACN_ADDRESS_ROOT_VECTOR = 18
+	ACN_ADDRESS_ROOT_VECTOR = 18,
+
+	// Low 12 bits = PDU length
+	// High 4 bits = 0x7
+	// ((ACN_HEADER_SIZE + ACN_DMX_SIZE) << 4) | ((VHD_V_FLAG | VHD_H_FLAG | VHD_D_FLAG >> 4)
+	ACN_PDU_FLAGS_AND_LENGHT_SIZE = 0x27e7
 };
 
 enum
diff --git a/Engine/Plugins/Experimental/DigitalMultiplex/DMXProtocol/Source/DMXProtocolSACN/Private/DMXProtocolSACNModule.cpp b/Engine/Plugins/Experimental/DigitalMultiplex/DMXProtocol/Source/DMXProtocolSACN/Private/DMXProtocolSACNModule.cpp
--- a/Engine/Plugins/Experimental/DigitalMultiplex/DMXProtocol/Source/DMXProtocolSACN/Private/DMXProtocolSACNModule.cpp
+++ b/Engine/Plugins/Experimental/DigitalMultiplex/DMXProtocol/Source/DMXProtocolSACN/Private/DMXProtocolSACNModule.cpp
@@ -9,20 +9,14 @@
 #include "DMXProtocolTypes.h"
 #include "Dom/JsonObject.h"
 
-#include "Managers/DMXProtocolDeviceManager.h"
-#include "Managers/DMXProtocolInterfaceManager.h"
-#include "Managers/DMXProtocolPortManager.h"
 #include "Managers/DMXProtocolUniverseManager.h"
 
-#include "DMXProtocolInterfaceSACN.h"
-#include "DMXProtocolDeviceSACN.h"
 #include "DMXProtocolUniverseSACN.h"
-#include "DMXProtocolPortSACN.h"
 #include "DMXProtocolSACNConstants.h"
 
 IMPLEMENT_MODULE(FDMXProtocolSACNModule, DMXProtocolSACN);
 
-const FName FDMXProtocolSACNModule::NAME_SACN = FName(TEXT("sACN"));
+const FName FDMXProtocolSACNModule::NAME_SACN = FName(DMX_PROTOCOLNAME_SACN);
 
 FAutoConsoleCommand FDMXProtocolSACNModule::SendDMXCommand(
 	TEXT("DMX.SACN.SendDMX"),
@@ -38,14 +32,14 @@ TSharedPtr<IDMXProtocol> FDMXProtocolFactorySACN::CreateProtocol(const FName & P
 	{
 		if (!ProtocolSACNPtr->Init())
 		{
-			UE_LOG_DMXPROTOCOL(Warning, TEXT("SACN failed to initialize!"));
+			UE_LOG_DMXPROTOCOL(Verbose, TEXT("SACN failed to initialize!"));
 			ProtocolSACNPtr->Shutdown();
 			ProtocolSACNPtr = nullptr;
 		}
 	}
 	else
 	{
-		UE_LOG_DMXPROTOCOL(Warning, TEXT("SACN disabled!"));
+		UE_LOG_DMXPROTOCOL(Verbose, TEXT("SACN disabled!"));
 		ProtocolSACNPtr->Shutdown();
 		ProtocolSACNPtr = nullptr;
 	}
@@ -83,7 +77,7 @@ void FDMXProtocolSACNModule::SendDMXCommandHandler(const TArray<FString>& Args)
 {
 	if (Args.Num() < 2)
 	{
-		UE_LOG_DMXPROTOCOL(Warning, TEXT("Not enough arguments. It won't be sent\n"
+		UE_LOG_DMXPROTOCOL(Verbose, TEXT("Not enough arguments. It won't be sent\n"
 			"Command structure is DMX.SACN.SendDMX [UniverseID] Channel:Value Channel:Value Channel:Value\n"
 			"For example: DMX.SACN.SendDMX 17 10:6 11:7 12:8 13:9"));
 		return;
@@ -93,7 +87,7 @@ void FDMXProtocolSACNModule::SendDMXCommandHandler(const TArray<FString>& Args)
 	LexTryParseString<uint32>(UniverseID, *Args[0]);
 	if (UniverseID > ACN_MAX_UNIVERSES)
 	{
-		UE_LOG_DMXPROTOCOL(Warning, TEXT("The UniverseID is bigger then the max universe value. It won't be sent. It won't be sent\n"
+		UE_LOG_DMXPROTOCOL(Verbose, TEXT("The UniverseID is bigger then the max universe value. It won't be sent. It won't be sent\n"
 			"For example: DMX.SACN.SendDMX 17 10:6 11:7 12:8 13:9\n"
 			"Where Universe %d should be less then %d"), UniverseID, ACN_MAX_UNIVERSES);
 		return;
@@ -113,7 +107,7 @@ void FDMXProtocolSACNModule::SendDMXCommandHandler(const TArray<FString>& Args)
 		LexTryParseString<uint32>(Key, *KeyStr);
 		if (Key > DMX_UNIVERSE_SIZE)
 		{
-			UE_LOG_DMXPROTOCOL(Warning, TEXT("The input channel is bigger then the universe size. It won't be sent\n"
+			UE_LOG_DMXPROTOCOL(Verbose, TEXT("The input channel is bigger then the universe size. It won't be sent\n"
 				"For example: DMX.SACN.SendDMX 17 10:6 11:7 12:8 13:9\n"
 				"Where channel %d should be less then %d"), Key, DMX_UNIVERSE_SIZE);
 			return;
@@ -121,7 +115,7 @@ void FDMXProtocolSACNModule::SendDMXCommandHandler(const TArray<FString>& Args)
 		LexTryParseString<uint32>(Value, *ValueStr);
 		if (Value > DMX_MAX_CHANNEL_VALUE)
 		{
-			UE_LOG_DMXPROTOCOL(Warning, TEXT("The input value is bigger then the universe size. It won't be sent\n"
+			UE_LOG_DMXPROTOCOL(Verbose, TEXT("The input value is bigger then the universe size. It won't be sent\n"
 				"For example: DMX.SACN.SendDMX 17 10:6 11:7 12:8 13:9\n"
 				"Where value %d should be less then %d"), Value, DMX_MAX_CHANNEL_VALUE);
 			return;
@@ -129,6 +123,9 @@ void FDMXProtocolSACNModule::SendDMXCommandHandler(const TArray<FString>& Args)
 		DMXFragment.Add(Key, Value);
 	}
 
-	FDMXProtocolSACN* DMXProtocol = IDMXProtocol::Get<FDMXProtocolSACN>(FDMXProtocolSACNModule::NAME_SACN);
-	DMXProtocol->SetDMXFragment(UniverseID, DMXFragment, true);
+	TSharedPtr<IDMXProtocol> DMXProtocol = IDMXProtocol::Get(FDMXProtocolSACNModule::NAME_SACN);
+	if (DMXProtocol.IsValid())
+	{
+		DMXProtocol->SendDMXFragmentCreate(UniverseID, DMXFragment);
+	}
 }
diff --git a/Engine/Plugins/Experimental/DigitalMultiplex/DMXProtocol/Source/DMXProtocolSACN/Private/DMXProtocolTransportSACN.cpp b/Engine/Plugins/Experimental/DigitalMultiplex/DMXProtocol/Source/DMXProtocolSACN/Private/DMXProtocolTransportSACN.cpp
--- a/Engine/Plugins/Experimental/DigitalMultiplex/DMXProtocol/Source/DMXProtocolSACN/Private/DMXProtocolTransportSACN.cpp
+++ b/Engine/Plugins/Experimental/DigitalMultiplex/DMXProtocol/Source/DMXProtocolSACN/Private/DMXProtocolTransportSACN.cpp
@@ -199,7 +199,7 @@ void FDMXProtocolReceiverSACN::Update(const FTimespan& SocketWaitTime)
 
 	while (Socket->HasPendingData(Size))
 	{
-		FArrayReaderPtr Reader = MakeShareable(new FArrayReader(true));
+		FArrayReaderPtr Reader = MakeShared<FArrayReader, ESPMode::ThreadSafe>(true);
 		Reader->SetNumUninitialized(FMath::Min(Size, 65507u));
 
 		int32 Read = 0;
diff --git a/Engine/Plugins/Experimental/DigitalMultiplex/DMXProtocol/Source/DMXProtocolSACN/Private/DMXProtocolTransportSACN.h b/Engine/Plugins/Experimental/DigitalMultiplex/DMXProtocol/Source/DMXProtocolSACN/Private/DMXProtocolTransportSACN.h
--- a/Engine/Plugins/Experimental/DigitalMultiplex/DMXProtocol/Source/DMXProtocolSACN/Private/DMXProtocolTransportSACN.h
+++ b/Engine/Plugins/Experimental/DigitalMultiplex/DMXProtocol/Source/DMXProtocolSACN/Private/DMXProtocolTransportSACN.h
@@ -3,9 +3,12 @@
 #pragma once
 
 #include "CoreMinimal.h"
-#include "Containers/Queue.h"
-#include "Interfaces/IDMXProtocolTransport.h"
+
 #include "DMXProtocolTypes.h"
+#include "Interfaces/IDMXProtocolTransport.h"
+
+#include "Containers/Queue.h"
+#include "HAL/CriticalSection.h"
 
 class FSocket;
 class ISocketSubsystem;
@@ -47,7 +50,6 @@ private:
 	/** Holds the last sent message number. */
 	int32 LastSentPackage;
 
-	/** */
 	mutable FCriticalSection CriticalSection;
 
 	FThreadSafeCounter StopTaskCounter;
diff --git a/Engine/Plugins/Experimental/DigitalMultiplex/DMXProtocol/Source/DMXProtocolSACN/Private/DMXProtocolUniverseSACN.cpp b/Engine/Plugins/Experimental/DigitalMultiplex/DMXProtocol/Source/DMXProtocolSACN/Private/DMXProtocolUniverseSACN.cpp
--- a/Engine/Plugins/Experimental/DigitalMultiplex/DMXProtocol/Source/DMXProtocolSACN/Private/DMXProtocolUniverseSACN.cpp
+++ b/Engine/Plugins/Experimental/DigitalMultiplex/DMXProtocol/Source/DMXProtocolSACN/Private/DMXProtocolUniverseSACN.cpp
@@ -3,74 +3,52 @@
 #include "DMXProtocolUniverseSACN.h"
 
 #include "DMXProtocolSACN.h"
-#include "Interfaces/IDMXProtocolDevice.h"
 #include "DMXProtocolConstants.h"
 #include "DMXProtocolTypes.h"
 #include "DMXProtocolTransportSACN.h"
+#include "DMXProtocolSettings.h"
+
 #include "Common/UdpSocketBuilder.h"
 
 #include "DMXProtocolSACNUtils.h"
 
-FDMXProtocolUniverseSACN::FDMXProtocolUniverseSACN(IDMXProtocol* InDMXProtocol, TSharedPtr<IDMXProtocolPort> InPort, uint16 InUniverseID)
-	: DMXProtocol(InDMXProtocol)
-	, Port(InPort)
-	, Priority(0)
-	, UniverseID(InUniverseID)
+FDMXProtocolUniverseSACN::FDMXProtocolUniverseSACN(TSharedPtr<IDMXProtocol> InDMXProtocol, const FJsonObject& InSettings)
+	: WeakDMXProtocol(InDMXProtocol)
+	, ListeningSocket(nullptr)
 {
-	checkf(DMXProtocol, TEXT("DMXProtocol pointer is nullptr"));
-	checkf(Port.IsValid(), TEXT("DMXProtocol port is not valid"));
+	NetworkErrorMessagePrefix = TEXT("NETWORK ERROR SACN:");
+
+	Settings = MakeShared<FJsonObject>(InSettings);
+
+	checkf(WeakDMXProtocol.IsValid(), TEXT("DMXProtocol pointer is not valid"));
+	checkf(Settings->HasField(TEXT("UniverseID")), TEXT("DMXProtocol UniverseID is not valid"));
+	UniverseID = Settings->GetNumberField(TEXT("UniverseID"));
 
 	OutputDMXBuffer = MakeShared<FDMXBuffer>();
 	InputDMXBuffer = MakeShared<FDMXBuffer>();
 
-	FIPv4Endpoint UnicastEndpoint = FIPv4Endpoint::Any;
-	TSharedPtr<FInternetAddr> MulticastAddr = FDMXProtocolSACN::GetUniverseAddr(UniverseID);
-	FIPv4Endpoint MulticastEndpoint = FIPv4Endpoint(MulticastAddr);
+	InterfaceIPAddress = GetDefault<UDMXProtocolSettings>()->InterfaceIPAddress;
 
-	ListenerSocket = FUdpSocketBuilder(TEXT("UDPSACNListenerSocket"))
-		.AsNonBlocking()
-		.AsReusable()
-#if PLATFORM_WINDOWS
-		.BoundToAddress(UnicastEndpoint.Address)
-#endif
-		.BoundToPort(MulticastEndpoint.Port)
-		.JoinedToGroup(MulticastEndpoint.Address, UnicastEndpoint.Address)
-		.WithMulticastLoopback()
-		.WithMulticastTtl(1)
-		.WithMulticastInterface(UnicastEndpoint.Address)
-		.WithReceiveBufferSize(2 * 1024 * 1024);
+	// Set Network Interface listener
+	NetworkInterfaceChangedHandle = IDMXProtocol::OnNetworkInterfaceChanged.AddRaw(this, &FDMXProtocolUniverseSACN::OnNetworkInterfaceChanged);
 
-	if (ListenerSocket == nullptr)
+	// Set Network Interface
+	FString ErrorMessage;
+	if (!RestartNetworkInterface(InterfaceIPAddress, ErrorMessage))
 	{
-		UE_LOG_DMXPROTOCOL(Error, TEXT("ERROR create BroadcastSocket"));
-	}
-	else
-	{
-		SACNReceiver = MakeShared<FDMXProtocolReceiverSACN>(*ListenerSocket, FTimespan::FromMilliseconds(100));
-		SACNReceiver->OnDataReceived().BindRaw(this, &FDMXProtocolUniverseSACN::OnDataReceived);
+		UE_LOG_DMXPROTOCOL(Error, TEXT("%s:%s"), NetworkErrorMessagePrefix, *ErrorMessage);
 	}
 }
 
 FDMXProtocolUniverseSACN::~FDMXProtocolUniverseSACN()
 {
-	SACNReceiver.Reset();
-
-	// Clear all sockets
-	if (ListenerSocket)
-	{
-		ListenerSocket->Close();
-		ISocketSubsystem::Get(PLATFORM_SOCKETSUBSYSTEM)->DestroySocket(ListenerSocket);
-	}
-}
-
-IDMXProtocol * FDMXProtocolUniverseSACN::GetProtocol() const
-{
-	return DMXProtocol;
+	ReleaseNetworkInterface();
+	IDMXProtocol::OnNetworkInterfaceChanged.Remove(NetworkInterfaceChangedHandle);
 }
 
-TWeakPtr<IDMXProtocolPort> FDMXProtocolUniverseSACN::GetCachedUniversePort() const
+TSharedPtr<IDMXProtocol> FDMXProtocolUniverseSACN::GetProtocol() const
 {
-	return Port;
+	return WeakDMXProtocol.Pin();
 }
 
 TSharedPtr<FDMXBuffer> FDMXProtocolUniverseSACN::GetOutputDMXBuffer() const
@@ -85,8 +63,7 @@ TSharedPtr<FDMXBuffer> FDMXProtocolUniverseSACN::GetInputDMXBuffer() const
 
 bool FDMXProtocolUniverseSACN::SetDMXFragment(const IDMXFragmentMap & DMXFragment)
 {
-	OutputDMXBuffer->SetDMXFragment(DMXFragment);
-	return false;
+	return OutputDMXBuffer->SetDMXFragment(DMXFragment);
 }
 
 uint8 FDMXProtocolUniverseSACN::GetPriority() const
@@ -94,13 +71,25 @@ uint8 FDMXProtocolUniverseSACN::GetPriority() const
 	return Priority;
 }
 
-uint16 FDMXProtocolUniverseSACN::GetUniverseID() const
+uint32 FDMXProtocolUniverseSACN::GetUniverseID() const
 {
 	return UniverseID;
 }
 
+TSharedPtr<FJsonObject> FDMXProtocolUniverseSACN::GetSettings() const
+{
+	return Settings;
+}
+
+bool FDMXProtocolUniverseSACN::IsSupportRDM() const
+{
+	return true;
+}
+
 void FDMXProtocolUniverseSACN::OnDataReceived(const FArrayReaderPtr & Buffer)
 {
+	FScopeLock Lock(&OnDataReceivedCS);
+
 	// It will be more handlers
 	switch (SACN::GetRootPacketType(Buffer))
 	{
@@ -122,12 +111,110 @@ bool FDMXProtocolUniverseSACN::HandleReplyPacket(const FArrayReaderPtr & Buffer)
 	// Make sure we copy same amount of data
 	if (InputDMXBuffer->GetDMXData().Num() == ACN_DMX_SIZE)
 	{
-		FMemory::Memcpy(InputDMXBuffer->GetDMXData().GetData(), IncomingDMXDMPLayer.DMX, ACN_DMX_SIZE);
+		InputDMXBuffer->SetDMXBuffer(IncomingDMXDMPLayer.DMX, ACN_DMX_SIZE);
+
+		GetProtocol()->GetOnUniverseInputUpdate().Broadcast(GetProtocol()->GetProtocolName(), GetUniverseID(), InputDMXBuffer->GetDMXData());
 		return true;
 	}
 	else
 	{
-		UE_LOG_DMXPROTOCOL(Warning, TEXT("Size of incoming DMX buffer is wrong"));
+		UE_LOG_DMXPROTOCOL(Error, TEXT("%s:Size of incoming DMX buffer is wrong"), NetworkErrorMessagePrefix);
 		return false;
 	}
 }
+
+void FDMXProtocolUniverseSACN::OnNetworkInterfaceChanged(const FString& InInterfaceIPAddress)
+{
+	FString ErrorMessage;
+	if (!RestartNetworkInterface(InInterfaceIPAddress, ErrorMessage))
+	{
+		UE_LOG_DMXPROTOCOL(Error, TEXT("%s:%s"), NetworkErrorMessagePrefix, *ErrorMessage);
+	}
+}
+
+bool FDMXProtocolUniverseSACN::RestartNetworkInterface(const FString& InInterfaceIPAddress, FString& OutErrorMessage)
+{
+	FScopeLock Lock(&ListeningSocketsCS);
+
+	// Clean the error message
+	OutErrorMessage.Empty();
+
+	// Try to create IP address at the first
+	ISocketSubsystem* SocketSubsystem = ISocketSubsystem::Get(PLATFORM_SOCKETSUBSYSTEM);
+	TSharedPtr<FInternetAddr> SenderAddr = SocketSubsystem->CreateInternetAddr();
+	bool bIsValid = false;
+	SenderAddr->SetIp(*InInterfaceIPAddress, bIsValid);
+    //SenderAddr->SetAnyAddress();
+	SenderAddr->SetPort(ACN_PORT);
+	if (!bIsValid)
+	{
+		OutErrorMessage = FString::Printf(TEXT("Wrong IP address: %s"), *InInterfaceIPAddress);
+		return false;
+	}
+
+	// Release old network interface
+	ReleaseNetworkInterface();
+
+	FIPv4Endpoint SenderEndpoint = FIPv4Endpoint(SenderAddr);
+	TSharedPtr<FInternetAddr> ListeningAddr = FDMXProtocolSACN::GetUniverseAddr(UniverseID);
+	FIPv4Endpoint ListeningEndpoint = FIPv4Endpoint(ListeningAddr);
+
+	FSocket* NewListeningSocket = FUdpSocketBuilder(TEXT("SACNListeningSocket"))
+		.AsNonBlocking()
+
+#if PLATFORM_WINDOWS
+		// For 0.0.0.0, Windows will pick the default interface instead of all
+		// interfaces. Here we allow to specify which interface to bind to. 
+		// On all other platforms we bind to the wildcard IP address in order
+		// to be able to also receive packets that were sent directly to the
+		// interface IP instead of the multicast address.
+		.BoundToAddress(SenderEndpoint.Address)
+#endif
+		.BoundToPort(SenderEndpoint.Port)
+#if PLATFORM_SUPPORTS_UDP_MULTICAST_GROUP
+		.JoinedToGroup(ListeningEndpoint.Address, SenderEndpoint.Address)
+		.WithMulticastLoopback()
+		.WithMulticastTtl(1)
+		.WithMulticastInterface(SenderEndpoint.Address)
+#endif
+        .AsReusable();
+
+
+	if (NewListeningSocket == nullptr)
+	{
+		OutErrorMessage = FString::Printf(TEXT("Error create ListeningSocket: %s"), *InterfaceIPAddress);
+		return false;
+	}
+
+	// Set new network interface IP
+	InterfaceIPAddress = InInterfaceIPAddress;
+
+	// Save New socket;
+	ListeningSocket = NewListeningSocket;
+
+	// Set new receiver
+	SACNReceiver = MakeShared<FDMXProtocolReceiverSACN>(*ListeningSocket, FTimespan::FromMilliseconds(100));
+	SACNReceiver->OnDataReceived().BindRaw(this, &FDMXProtocolUniverseSACN::OnDataReceived);
+
+	return true;
+}
+
+void FDMXProtocolUniverseSACN::ReleaseNetworkInterface()
+{
+	if (SACNReceiver.IsValid())
+	{
+		SACNReceiver.Reset();
+		SACNReceiver = nullptr;
+	}
+
+	// Clean all sockets
+	if (ListeningSocket != nullptr)
+	{
+		ListeningSocket->Close();
+		if (ISocketSubsystem* SocketSubsystem = ISocketSubsystem::Get(PLATFORM_SOCKETSUBSYSTEM))
+		{
+			SocketSubsystem->DestroySocket(ListeningSocket);
+		}
+		ListeningSocket = nullptr;
+	}
+}
