- DMX Editor

diff --git a/Engine/Plugins/Experimental/DigitalMultiplex/DMXProtocol/Source/DMXProtocolSACN/Private/DMXProtocolUniverseSACN.h b/Engine/Plugins/Experimental/DigitalMultiplex/DMXProtocol/Source/DMXProtocolSACN/Private/DMXProtocolUniverseSACN.h
--- a/Engine/Plugins/Experimental/DigitalMultiplex/DMXProtocol/Source/DMXProtocolSACN/Private/DMXProtocolUniverseSACN.h
+++ b/Engine/Plugins/Experimental/DigitalMultiplex/DMXProtocol/Source/DMXProtocolSACN/Private/DMXProtocolUniverseSACN.h
@@ -3,29 +3,41 @@
 #pragma once
 
 #include "CoreMinimal.h"
+
 #include "Interfaces/IPv4/IPv4Endpoint.h"
 #include "Interfaces/IDMXProtocolUniverse.h"
+#include "Interfaces/IDMXNetworkInterface.h"
 #include "Packets/DMXProtocolE131PDUPacket.h"
 
+#include "HAL/CriticalSection.h"
+
 class FSocket;
 
 class DMXPROTOCOLSACN_API FDMXProtocolUniverseSACN
 	: public IDMXProtocolUniverse
+	, public IDMXNetworkInterface
 {
 public:
-	FDMXProtocolUniverseSACN(IDMXProtocol* InDMXProtocol, TSharedPtr<IDMXProtocolPort> InPort, uint16 InUniverseID);
+	FDMXProtocolUniverseSACN(TSharedPtr<IDMXProtocol> InDMXProtocol, const FJsonObject& InSettings);
 	virtual ~FDMXProtocolUniverseSACN();
 
 	//~ Begin IDMXProtocolDevice implementation
-	virtual IDMXProtocol* GetProtocol() const override;
-	virtual TWeakPtr<IDMXProtocolPort> GetCachedUniversePort() const override;
-	virtual TSharedPtr<FDMXBuffer> GetOutputDMXBuffer() const override;
+	virtual TSharedPtr<IDMXProtocol> GetProtocol() const override;
 	virtual TSharedPtr<FDMXBuffer> GetInputDMXBuffer() const override;
+	virtual TSharedPtr<FDMXBuffer> GetOutputDMXBuffer() const override;
 	virtual bool SetDMXFragment(const IDMXFragmentMap& DMXFragment) override;
 	virtual uint8 GetPriority() const override;
-	virtual uint16 GetUniverseID() const override;
+	virtual uint32 GetUniverseID() const override;
+	virtual TSharedPtr<FJsonObject> GetSettings() const override;
+	virtual bool IsSupportRDM() const override;
 	//~ End IDMXProtocolDevice implementation
 
+	//~ Begin IDMXNetworkInterface implementation
+	virtual void OnNetworkInterfaceChanged(const FString& InInterfaceIPAddress) override;
+	virtual bool RestartNetworkInterface(const FString& InInterfaceIPAddress, FString& OutErrorMessage) override;
+	virtual void ReleaseNetworkInterface() override;
+	//~ End IDMXNetworkInterface implementation
+
 	//~ SACN Getters functions
 	const FDMXProtocolE131RootLayerPacket& GetIncomingDMXRootLayer() const { return IncomingDMXRootLayer; }
 	const FDMXProtocolE131FramingLayerPacket& GetIncomingDMXFramingLayer() const { return IncomingDMXFramingLayer; }
@@ -37,17 +49,29 @@ private:
 	bool HandleReplyPacket(const FArrayReaderPtr & Buffer);
 
 private:
-	IDMXProtocol* DMXProtocol;
-	TWeakPtr<IDMXProtocolPort> Port;
+	TWeakPtr<IDMXProtocol> WeakDMXProtocol;
 	TSharedPtr<FDMXBuffer> OutputDMXBuffer;
 	TSharedPtr<FDMXBuffer> InputDMXBuffer;
 	uint8 Priority;
-	uint16 UniverseID;
+	uint32 UniverseID;
+	bool bIsRDMSupport;
+	TSharedPtr<FJsonObject> Settings;
 
-	FSocket* ListenerSocket;
+	mutable FCriticalSection OnDataReceivedCS;
+
+	FSocket* ListeningSocket;
 	TSharedPtr<IDMXProtocolReceiver> SACNReceiver;
 
 	FDMXProtocolE131RootLayerPacket IncomingDMXRootLayer;
 	FDMXProtocolE131FramingLayerPacket IncomingDMXFramingLayer;
 	FDMXProtocolE131DMPLayerPacket IncomingDMXDMPLayer;
+
+	/** Mutex protecting access to the sockets. */
+	mutable FCriticalSection ListeningSocketsCS;
+
+	FString InterfaceIPAddress;
+
+	FDelegateHandle NetworkInterfaceChangedHandle;
+
+	const TCHAR* NetworkErrorMessagePrefix;
 };
\ No newline at end of file
diff --git a/Engine/Plugins/Experimental/DigitalMultiplex/DMXProtocol/Source/DMXProtocolSACN/Private/Packets/DMXProtocolE131PDUPacket.cpp b/Engine/Plugins/Experimental/DigitalMultiplex/DMXProtocol/Source/DMXProtocolSACN/Private/Packets/DMXProtocolE131PDUPacket.cpp
--- a/Engine/Plugins/Experimental/DigitalMultiplex/DMXProtocol/Source/DMXProtocolSACN/Private/Packets/DMXProtocolE131PDUPacket.cpp
+++ b/Engine/Plugins/Experimental/DigitalMultiplex/DMXProtocol/Source/DMXProtocolSACN/Private/Packets/DMXProtocolE131PDUPacket.cpp
@@ -2,7 +2,7 @@
 
 #include "Packets/DMXProtocolE131PDUPacket.h"
 #include "DMXProtocolTypes.h"
-#include "Serialization/DMXMemoryWriter.h"
+#include "Serialization/BufferArchive.h"
 
 #include "DMXProtocolMacros.h"
 
@@ -14,7 +14,7 @@ REGISTER_DMX_ARCHIVE(FDMXProtocolUDPE131DiscoveryLayerPacket);
 
 TSharedPtr<FBufferArchive> FDMXProtocolE131RootLayerPacket::Pack()
 {
-	TSharedPtr<FDMXMemoryWriter> Writer = MakeShared<FDMXMemoryWriter>();
+	TSharedPtr<FBufferArchive> Writer = MakeShared<FBufferArchive>();
 	*Writer << *this;
 
 	return Writer;
@@ -28,8 +28,7 @@ void FDMXProtocolE131RootLayerPacket::Serialize(FArchive & Ar)
 	Ar.SetByteSwapping(false);
 	Ar.Serialize((void*)ACNPacketIdentifier, ACN_IDENTIFIER_SIZE);
 	Ar.SetByteSwapping(true);
-	Ar << Flags;
-	Ar << Length;
+	Ar << FlagsAndLength;
 	Ar << Vector;
 	Ar.SetByteSwapping(false);
 	Ar.Serialize((void*)CID, ACN_CIDBYTES);
@@ -38,7 +37,7 @@ void FDMXProtocolE131RootLayerPacket::Serialize(FArchive & Ar)
 
 TSharedPtr<FBufferArchive> FDMXProtocolE131FramingLayerPacket::Pack()
 {
-	TSharedPtr<FDMXMemoryWriter> Writer = MakeShared<FDMXMemoryWriter>();
+	TSharedPtr<FBufferArchive> Writer = MakeShared<FBufferArchive>();
 	*Writer << *this;
 
 	return Writer;
@@ -47,8 +46,7 @@ TSharedPtr<FBufferArchive> FDMXProtocolE131FramingLayerPacket::Pack()
 void FDMXProtocolE131FramingLayerPacket::Serialize(FArchive & Ar)
 {
 	Ar.SetByteSwapping(true);
-	Ar << Flags;
-	Ar << Length;
+	Ar << FlagsAndLength;
 	Ar << Vector;
 	Ar.SetByteSwapping(false);
 	Ar.Serialize((void*)SourceName, ACN_SOURCE_NAME_SIZE);
@@ -63,7 +61,7 @@ void FDMXProtocolE131FramingLayerPacket::Serialize(FArchive & Ar)
 
 TSharedPtr<FBufferArchive> FDMXProtocolE131DMPLayerPacket::Pack()
 {
-	TSharedPtr<FDMXMemoryWriter> Writer = MakeShared<FDMXMemoryWriter>();
+	TSharedPtr<FBufferArchive> Writer = MakeShared<FBufferArchive>();
 	*Writer << *this;
 
 	return Writer;
@@ -72,8 +70,7 @@ TSharedPtr<FBufferArchive> FDMXProtocolE131DMPLayerPacket::Pack()
 void FDMXProtocolE131DMPLayerPacket::Serialize(FArchive & Ar)
 {
 	Ar.SetByteSwapping(true);
-	Ar << Flags;
-	Ar << Length;
+	Ar << FlagsAndLength;
 	Ar << Vector;
 	Ar << AddressTypeAndDataType;
 	Ar << FirstPropertyAddress;
@@ -87,7 +84,7 @@ void FDMXProtocolE131DMPLayerPacket::Serialize(FArchive & Ar)
 
 TSharedPtr<FBufferArchive> FDMXProtocolUDPE131FramingLayerPacket::Pack()
 {
-	TSharedPtr<FDMXMemoryWriter> Writer = MakeShared<FDMXMemoryWriter>();
+	TSharedPtr<FBufferArchive> Writer = MakeShared<FBufferArchive>();
 	*Writer << *this;
 
 	return Writer;
@@ -96,8 +93,7 @@ TSharedPtr<FBufferArchive> FDMXProtocolUDPE131FramingLayerPacket::Pack()
 void FDMXProtocolUDPE131FramingLayerPacket::Serialize(FArchive & Ar)
 {
 	Ar.SetByteSwapping(true);
-	Ar << Flags;
-	Ar << Length;
+	Ar << FlagsAndLength;
 	Ar << Vector;
 	Ar.SetByteSwapping(false);
 	Ar.Serialize((void*)SourceName, ACN_SOURCE_NAME_SIZE);
@@ -106,7 +102,7 @@ void FDMXProtocolUDPE131FramingLayerPacket::Serialize(FArchive & Ar)
 
 TSharedPtr<FBufferArchive> FDMXProtocolUDPE131DiscoveryLayerPacket::Pack()
 {
-	TSharedPtr<FDMXMemoryWriter> Writer = MakeShared<FDMXMemoryWriter>();
+	TSharedPtr<FBufferArchive> Writer = MakeShared<FBufferArchive>();
 	*Writer << *this;
 
 	return Writer;
@@ -115,8 +111,7 @@ TSharedPtr<FBufferArchive> FDMXProtocolUDPE131DiscoveryLayerPacket::Pack()
 void FDMXProtocolUDPE131DiscoveryLayerPacket::Serialize(FArchive & Ar)
 {
 	Ar.SetByteSwapping(true);
-	Ar << Flags;
-	Ar << Length;
+	Ar << FlagsAndLength;
 	Ar << Vector;
 	Ar << Page;
 	Ar << Last;
diff --git a/Engine/Plugins/Experimental/DigitalMultiplex/DMXProtocol/Source/DMXProtocolSACN/Private/Packets/DMXProtocolE131PDUPacket.h b/Engine/Plugins/Experimental/DigitalMultiplex/DMXProtocol/Source/DMXProtocolSACN/Private/Packets/DMXProtocolE131PDUPacket.h
--- a/Engine/Plugins/Experimental/DigitalMultiplex/DMXProtocol/Source/DMXProtocolSACN/Private/Packets/DMXProtocolE131PDUPacket.h
+++ b/Engine/Plugins/Experimental/DigitalMultiplex/DMXProtocol/Source/DMXProtocolSACN/Private/Packets/DMXProtocolE131PDUPacket.h
@@ -21,8 +21,7 @@ public:
 	uint16 PreambleSize = ACN_RLP_PREAMBLE_SIZE;
 	uint16 PostambleSize = 0;
 	uint8 ACNPacketIdentifier[ACN_IDENTIFIER_SIZE] = { 'A', 'S', 'C', '-', 'E', '1', '.', '1', '7', '\0', '\0', '\0' };
-	uint8 Flags = 0x72;
-	uint8 Length = 0x6e;
+	uint16 FlagsAndLength = ACN_PDU_FLAGS_AND_LENGHT_SIZE;
 	uint32 Vector = VECTOR_ROOT_E131_DATA;
 	uint8 CID[ACN_CIDBYTES] = { 0 };
 };
@@ -39,15 +38,14 @@ public:
 	friend FArchive& operator<<(FArchive& Ar, FDMXProtocolE131FramingLayerPacket& Packet);
 
 public:
-	uint8 Flags = 0;
-	uint8 Length = 0;
+	uint16 FlagsAndLength = ACN_PDU_FLAGS_AND_LENGHT_SIZE;
 	uint32 Vector = VECTOR_E131_DATA_PACKET;
 	uint8 SourceName[ACN_SOURCE_NAME_SIZE] = "New Source";
-	uint8 Priority = 0;
+	uint8 Priority = 100;
 	uint16 SynchronizationAddress = 0;
 	uint8 SequenceNumber = 0;
 	uint8 Options = 0;
-	uint8 Universe = 0;
+	uint16 Universe = 0;
 };
 
 struct FDMXProtocolE131DMPLayerPacket
@@ -62,10 +60,9 @@ public:
 	friend FArchive& operator<<(FArchive& Ar, FDMXProtocolE131DMPLayerPacket& Packet);
 
 public:
-	uint8 Flags = 0;
-	uint8 Length = 0;
+	uint16 FlagsAndLength = ACN_PDU_FLAGS_AND_LENGHT_SIZE;
 	uint8 Vector = VECTOR_DMP_SET_PROPERTY;
-	uint8 AddressTypeAndDataType = 0;
+	uint8 AddressTypeAndDataType = 0xa1; // The const from SACN protocol documentation
 	uint16 FirstPropertyAddress = 0;
 	uint16 AddressIncrement = 0;
 	uint16 PropertyValueCount = 0;
@@ -85,8 +82,7 @@ public:
 	friend FArchive& operator<<(FArchive& Ar, FDMXProtocolUDPE131FramingLayerPacket& Packet);
 
 public:
-	uint8 Flags = 0x72;
-	uint8 Length = 0x0b;
+	uint16 FlagsAndLength = ACN_PDU_FLAGS_AND_LENGHT_SIZE;
 	uint32 Vector = 2;
 	uint8 SourceName[ACN_SOURCE_NAME_SIZE] = { 0 };
 	uint8 Reserved[4] = { 0 };
@@ -105,8 +101,7 @@ public:
 	friend FArchive& operator<<(FArchive& Ar, FDMXProtocolUDPE131DiscoveryLayerPacket& Packet);
 
 public:
-	uint8 Flags = 0;
-	uint8 Length = 0;
+	uint16 FlagsAndLength = ACN_PDU_FLAGS_AND_LENGHT_SIZE;
 	uint32 Vector = VECTOR_E131_EXTENDED_DISCOVERY;
 	uint8 Page = 0;
 	uint8 Last = 0;
diff --git a/Engine/Plugins/Experimental/DigitalMultiplex/DMXProtocol/Source/DMXProtocolSACN/Private/Tests/DMXPrtotocolSACNFlowTest.cpp b/Engine/Plugins/Experimental/DigitalMultiplex/DMXProtocol/Source/DMXProtocolSACN/Private/Tests/DMXPrtotocolSACNFlowTest.cpp
--- a/Engine/Plugins/Experimental/DigitalMultiplex/DMXProtocol/Source/DMXProtocolSACN/Private/Tests/DMXPrtotocolSACNFlowTest.cpp
+++ b/Engine/Plugins/Experimental/DigitalMultiplex/DMXProtocol/Source/DMXProtocolSACN/Private/Tests/DMXPrtotocolSACNFlowTest.cpp
@@ -8,88 +8,36 @@
 #include "DMXProtocolSACN.h"
 
 #include "DMXProtocolSACN.h"
-#include "DMXProtocolPortSACN.h"
-
-#include "Managers/DMXProtocolDeviceManager.h"
-#include "Managers/DMXProtocolInterfaceManager.h"
-#include "Managers/DMXProtocolPortManager.h"
-#include "Managers/DMXProtocolUniverseManager.h"
-
-#include "DMXProtocolInterfaceSACN.h"
-#include "DMXProtocolDeviceSACN.h"
 #include "DMXProtocolUniverseSACN.h"
-#include "DMXProtocolPortSACN.h"
-#include "DMXProtocolSACNConstants.h"
-
-
-struct FDMXProtocolLauncher
-{
-	FDMXProtocolLauncher(uint8 InDeviceID, uint8 InPortID)
-		: DeviceID(InDeviceID)
-		, PortID(InPortID)
-		, InterfaceID(0)
-		, UniverseID(3)
-	{
-		DMXProtocol = IDMXProtocol::Get<FDMXProtocolSACN>(FDMXProtocolSACNModule::NAME_SACN);
-
-		InterfaceSACN = MakeShared<FDMXProtocolInterfaceSACN>(DMXProtocol);
-		DMXProtocol->GetInterfaceManager()->AddInterface(0, InterfaceSACN);
-
-		FJsonObject DeviceSettings;
-		ProtocolDevice = MakeShared<FDMXProtocolDeviceSACN>(DMXProtocol, InterfaceSACN, DeviceSettings, DeviceID);
-		DMXProtocol->GetDeviceManager()->AddDevice(ProtocolDevice);
-
-		ProtocolOutputPort = MakeShared<FDMXProtocolPortSACN>(DMXProtocol, ProtocolDevice, PortID, UniverseID, DMX_PORT_OUTPUT);
-		DMXProtocol->GetPortManager()->AddOutputPort(ProtocolDevice, ProtocolOutputPort);
-
-		FJsonObject UniverseSettings;
-		ProtocolUniverse = MakeShared<FDMXProtocolUniverseSACN>(DMXProtocol, ProtocolOutputPort, UniverseID);
-		DMXProtocol->GetUniverseManager()->AddUniverse(ProtocolOutputPort, ProtocolUniverse);
-		ProtocolOutputPort->SetUniverse(ProtocolUniverse);
-	}
-
-	~FDMXProtocolLauncher()
-	{
-		DMXProtocol->GetInterfaceManager()->RemoveInterface(InterfaceID);
-		DMXProtocol->GetDeviceManager()->RemoveDevice(ProtocolDevice);
-		DMXProtocol->GetPortManager()->RemoveOutputPort(ProtocolDevice, ProtocolOutputPort);
-		DMXProtocol->GetUniverseManager()->RemoveUniverse(ProtocolOutputPort, ProtocolUniverse);
-	}
-
-	const uint8 DeviceID;
-	const uint8 PortID;
-	const uint8 InterfaceID;
-	const uint16 UniverseID;
-	TSharedPtr<FDMXProtocolInterfaceSACN> InterfaceSACN;
-	TSharedPtr<FDMXProtocolDeviceSACN> ProtocolDevice;
-	TSharedPtr<FDMXProtocolPortSACN> ProtocolOutputPort;
-	TSharedPtr<FDMXProtocolUniverseSACN> ProtocolUniverse;
-	FDMXProtocolSACN* DMXProtocol;
-};
 
+static uint16 UniverseValue = 34;
 
 IMPLEMENT_SIMPLE_AUTOMATION_TEST(DMXPrtotocolSACNBasicFlowTest, "VirtualProduction.DMX.SACN.BasicFlow", EAutomationTestFlags::ApplicationContextMask | EAutomationTestFlags::EngineFilter)
 
-static void SensDMXFragment(FAutomationTestBase& Test, FDMXProtocolSACN* DMXProtocol, TSharedPtr<FDMXProtocolDeviceSACN> InProtocolDevice, uint8 InPortID)
+bool DMXPrtotocolSACNBasicFlowTest::RunTest(const FString& Parameters)
 {
-	IDMXProtocolPort* SACNPort = DMXProtocol->GetPortManager()->GetPortByDeviceAndID(InProtocolDevice, InPortID, EDMXPortDirection::DMX_PORT_OUTPUT);
-	Test.TestNotNull(TEXT("SACNPort is valid"), SACNPort);
-	Test.TestEqual(TEXT("Universe ID should be equal 3"), SACNPort->GetUniverseID(), 3);
+	// Init Protocol and Managers
+	FDMXProtocolSACN* DMXProtocol = static_cast<FDMXProtocolSACN*>((IDMXProtocol::Get(FDMXProtocolSACNModule::NAME_SACN)).Get());
+
+	FJsonObject UniverseSettings;
+	UniverseSettings.SetNumberField(TEXT("UniverseID"), UniverseValue);
+	DMXProtocol->AddUniverse(UniverseSettings);
+	TSharedPtr<IDMXProtocolUniverse, ESPMode::ThreadSafe> Universe = DMXProtocol->GetUniverseById(UniverseValue);
 
+	TestTrue(TEXT("Universe is valid"), Universe.IsValid());
+	TestEqual(TEXT("Universe ID should be equal"), Universe->GetUniverseID(), UniverseValue);
 
 	// Send DMX fragment 1
-	const uint8 FixtureChannels[6] = { 0, 1, 2, 5, 6, 7 };
+	const uint8 FixtureChannels[6] = { 1, 2, 3, 4, 5, 6 };
 	const uint8 FixtureValues[6] = { 255, 155, 50, 100, 200, 220 };
 	IDMXFragmentMap DMXFragmentMap;
 	DMXFragmentMap.Add(FixtureChannels[0], FixtureValues[0]);
 	DMXFragmentMap.Add(FixtureChannels[1], FixtureValues[1]);
 	DMXFragmentMap.Add(FixtureChannels[2], FixtureValues[2]);
-	DMXProtocol->SetDMXFragment(SACNPort->GetUniverseID(), DMXFragmentMap);
+	DMXProtocol->SendDMXFragment(Universe->GetUniverseID(), DMXFragmentMap);
 
-	// Wait before recieve DMX
+	// Wait before receive DMX
 	FPlatformProcess::Sleep(0.2f);
-	TSharedPtr<IDMXProtocolUniverse> Universe = SACNPort->GetCachedUniverse().Pin();
-	Test.TestTrue(TEXT("Universe should be valid"), Universe.IsValid());
 	TSharedPtr<FDMXBuffer> InputDMXBuffer = Universe->GetInputDMXBuffer();
 
 	// Test package ACNPacketIdentifier
@@ -98,39 +46,31 @@ static void SensDMXFragment(FAutomationTestBase& Test, FDMXProtocolSACN* DMXProt
 		FDMXProtocolUniverseSACN* DMXProtocolUniverseSACN = static_cast<FDMXProtocolUniverseSACN*>(Universe.Get());
 		FString ExpectedInentifier = FString::Printf(TEXT("%s"), ACNPacketIdentifier);
 		FString IncomingInentifier = FString::Printf(TEXT("%s"), DMXProtocolUniverseSACN->GetIncomingDMXRootLayer().ACNPacketIdentifier);
-		Test.TestTrue(TEXT("Incoming ACNPacketIdentifier should be the same"), ExpectedInentifier.Equals(IncomingInentifier));
+		TestTrue(TEXT("Incoming ACNPacketIdentifier should be the same"), ExpectedInentifier.Equals(IncomingInentifier));
 	}
 
-	Test.TestEqual(TEXT("Incoming buffer should be same"), FixtureValues[0], InputDMXBuffer->GetDMXData()[FixtureChannels[0]]);
-	Test.TestEqual(TEXT("Incoming buffer should be same"), FixtureValues[1], InputDMXBuffer->GetDMXData()[FixtureChannels[1]]);
-	Test.TestEqual(TEXT("Incoming buffer should be same"), FixtureValues[2], InputDMXBuffer->GetDMXData()[FixtureChannels[2]]);
+	TestEqual(TEXT("Incoming buffer should be same"), FixtureValues[0], InputDMXBuffer->GetDMXData()[FixtureChannels[0] - 1]);
+	TestEqual(TEXT("Incoming buffer should be same"), FixtureValues[1], InputDMXBuffer->GetDMXData()[FixtureChannels[1] - 1]);
+	TestEqual(TEXT("Incoming buffer should be same"), FixtureValues[2], InputDMXBuffer->GetDMXData()[FixtureChannels[2] - 1]);
 
 	// Send DMX fragment 2
 	DMXFragmentMap.Add(FixtureChannels[3], FixtureValues[3]);
 	DMXFragmentMap.Add(FixtureChannels[4], FixtureValues[4]);
 	DMXFragmentMap.Add(FixtureChannels[5], FixtureValues[5]);
-	DMXProtocol->SetDMXFragment(SACNPort->GetUniverseID(), DMXFragmentMap);
+	DMXProtocol->SendDMXFragment(Universe->GetUniverseID(), DMXFragmentMap);
 	FPlatformProcess::Sleep(0.2f);
 
-	Test.TestEqual(TEXT("Old value DMX input should be the same"), FixtureValues[0], InputDMXBuffer->GetDMXData()[FixtureChannels[0]]);
-	Test.TestEqual(TEXT("Incoming buffer should be same"), FixtureValues[4], InputDMXBuffer->GetDMXData()[FixtureChannels[4]]);
-	Test.TestEqual(TEXT("Incoming buffer should be same"), FixtureValues[5], InputDMXBuffer->GetDMXData()[FixtureChannels[5]]);
+	TestEqual(TEXT("Old value DMX input should be the same"), FixtureValues[0], InputDMXBuffer->GetDMXData()[FixtureChannels[0] - 1]);
+	TestEqual(TEXT("Incoming buffer should be same"), FixtureValues[4], InputDMXBuffer->GetDMXData()[FixtureChannels[4] - 1]);
+	TestEqual(TEXT("Incoming buffer should be same"), FixtureValues[5], InputDMXBuffer->GetDMXData()[FixtureChannels[5] - 1]);
 
 	// @TODO Implement test cases for Discovery
 	TArray<uint16> Universes;
-	for (int UniversesIndex = 0; UniversesIndex <= DMX_UNIVERSE_SIZE; UniversesIndex++)
+	for (int32 UniversesIndex = 0; UniversesIndex < DMX_UNIVERSE_SIZE; UniversesIndex++)
 	{
 		Universes.Add(UniversesIndex);
 	}
 	DMXProtocol->SendDiscovery(Universes);
-}
-
-bool DMXPrtotocolSACNBasicFlowTest::RunTest(const FString& Parameters)
-{
-	FDMXProtocolLauncher DMXProtocolLauncher(0, 0);
-
-	// Test Sending DMX and Discovery
-	SensDMXFragment(*this, DMXProtocolLauncher.DMXProtocol, DMXProtocolLauncher.ProtocolDevice, DMXProtocolLauncher.PortID);
 
 	return true;
 }
@@ -139,7 +79,13 @@ IMPLEMENT_SIMPLE_AUTOMATION_TEST(FDMXPrtotocolSACNConsoleCommandsTest, "VirtualP
 
 bool FDMXPrtotocolSACNConsoleCommandsTest::RunTest(const FString& Parameters)
 {
-	FDMXProtocolLauncher DMXProtocolLauncher(0, 0);
+	// Init Protocol and Managers
+	FDMXProtocolSACN* DMXProtocol = static_cast<FDMXProtocolSACN*>((IDMXProtocol::Get(FDMXProtocolSACNModule::NAME_SACN)).Get());
+
+	FJsonObject UniverseSettings;
+	UniverseSettings.SetNumberField(TEXT("UniverseID"), UniverseValue);
+	DMXProtocol->AddUniverse(UniverseSettings);
+	TSharedPtr<IDMXProtocolUniverse, ESPMode::ThreadSafe> Universe = DMXProtocol->GetUniverseById(UniverseValue);
 
 	// DMX.SACN.SendDMX[UniverseID] Channel:Value Channel : Value Channel : Value ...
 	uint32 Values[3][2] =
@@ -151,7 +97,7 @@ bool FDMXPrtotocolSACNConsoleCommandsTest::RunTest(const FString& Parameters)
 
 	FString Command(
 		FString::Printf(TEXT("DMX.SACN.SendDMX %d %d:%d %d:%d %d:%d"),
-			DMXProtocolLauncher.ProtocolOutputPort->GetUniverseID(),
+			Universe->GetUniverseID(),
 			Values[0][0], Values[0][1],
 			Values[1][0], Values[1][1],
 			Values[2][0], Values[2][1]
@@ -161,11 +107,10 @@ bool FDMXPrtotocolSACNConsoleCommandsTest::RunTest(const FString& Parameters)
 
 	FPlatformProcess::Sleep(0.2f);
 
-	TSharedPtr<IDMXProtocolUniverse> Universe = DMXProtocolLauncher.ProtocolOutputPort->GetCachedUniverse().Pin();
 	TSharedPtr<FDMXBuffer> InputDMXBuffer = Universe->GetInputDMXBuffer();
-	TestEqual(TEXT("Incoming buffer should be same"), Values[0][1], InputDMXBuffer->GetDMXData()[Values[0][0]]);
-	TestEqual(TEXT("Incoming buffer should be same"), Values[1][1], InputDMXBuffer->GetDMXData()[Values[1][0]]);
-	TestEqual(TEXT("Incoming buffer should be same"), Values[2][1], InputDMXBuffer->GetDMXData()[Values[2][0]]);
+	TestEqual(TEXT("Incoming buffer should be same"), Values[0][1], InputDMXBuffer->GetDMXData()[Values[0][0] - 1]);
+	TestEqual(TEXT("Incoming buffer should be same"), Values[1][1], InputDMXBuffer->GetDMXData()[Values[1][0] - 1]);
+	TestEqual(TEXT("Incoming buffer should be same"), Values[2][1], InputDMXBuffer->GetDMXData()[Values[2][0] - 1]);
 
 	return true;
 }
\ No newline at end of file
