Edigrating CL 12099511, 12098829, 12099511

Fixing the threads and sockets issue for SACN

- Remove the thread for each universe
- Open the socket only when the universe is using
- Clean the sockets if there are no listeners
- Remove SACN thread listener
- Using gameplay tick for listening changes on the sockets

Fix for incremental builds

diff --git a/Engine/Plugins/Experimental/DigitalMultiplex/DMXEngine/Source/DMXEditor/Private/Widgets/OutputFader/SDMXFader.cpp b/Engine/Plugins/Experimental/DigitalMultiplex/DMXEngine/Source/DMXEditor/Private/Widgets/OutputFader/SDMXFader.cpp
--- a/Engine/Plugins/Experimental/DigitalMultiplex/DMXEngine/Source/DMXEditor/Private/Widgets/OutputFader/SDMXFader.cpp
+++ b/Engine/Plugins/Experimental/DigitalMultiplex/DMXEngine/Source/DMXEditor/Private/Widgets/OutputFader/SDMXFader.cpp
@@ -6,6 +6,7 @@
 #include "DMXEditor.h"
 #include "Library/DMXEntityFader.h"
 #include "DMXEditorLog.h"
+#include "DMXProtocolCommon.h"
 
 #include "Widgets/Common/SSpinBoxVertical.h"
 #include "Widgets/OutputFader/SDMXOutputFaderList.h"
@@ -23,6 +24,9 @@ void SDMXFader::Construct(const FArguments& InArgs)
 	WeakDMXEditor = InArgs._DMXEditor;
 	CurrentFaderValue = 0;
 
+	OnValueChanged = InArgs._OnValueChanged;
+	OnSendStateChanged = InArgs._OnSendStateChanged;
+
 	ChildSlot
 		.Padding(0.f)
 		[
@@ -100,6 +104,7 @@ void SDMXFader::Construct(const FArguments& InArgs)
 								[
 									SAssignNew(SendDMXCheckBox, SCheckBox)
 									.IsChecked(ECheckBoxState::Checked)
+									.OnCheckStateChanged(this, &SDMXFader::HandleSendDMXCheckChanged)
 								]
 							]
 							+ SVerticalBox::Slot()
@@ -244,6 +249,11 @@ FReply SDMXFader::HandleRemoveFaderClicked()
 	return FReply::Handled();
 }
 
+void SDMXFader::HandleSendDMXCheckChanged(ECheckBoxState NewState)
+{
+
+}
+
 FReply SDMXFader::OnMouseButtonDown(const FGeometry& MyGeometry, const FPointerEvent& MouseEvent)
 {
 	if (MouseEvent.GetEffectingButton() == EKeys::LeftMouseButton)
@@ -275,13 +285,12 @@ void SDMXFader::SetFaderLabel(const FString& InLabel)
 	CustomFaderLabel->SetText(FText::FromString(InLabel));
 }
 
-void SDMXFader::AddChannelWidget(const FString& InUniverse, const FString& InChannel, uint16 InUniverseNumber, uint32 InChannelNumber, const IDMXProtocolPtr& InDMXProtocol)
+void SDMXFader::AddChannelWidget(const FString& InUniverse, const FString& InChannel, uint16 InUniverseNumber, uint32 InChannelNumber)
 {
 	TSharedPtr<SDMXFaderChannel> NewChannel = SNew(SDMXFaderChannel)
 		.DMXEditor(WeakDMXEditor)
 		.UniverseNumber(InUniverseNumber)
-		.ChannelNumber(InChannelNumber)
-		.DMXProtocol(InDMXProtocol);
+		.ChannelNumber(InChannelNumber);
 
 	FaderChannels.Add(NewChannel);
 
@@ -325,6 +334,11 @@ void SDMXFader::SelectThisFader()
 	}
 }
 
+bool SDMXFader::ShouldSendDMX() const
+{
+	return SendDMXCheckBox->IsChecked();
+}
+
 const FSlateBrush* SDMXFader::GetBorderImage() const
 {
 	if (IsHovered())
@@ -341,29 +355,18 @@ void SDMXFader::HandleFaderChanged(uint8 NewValue)
 {
 	CurrentFaderValue = NewValue;
 
-	if (SendDMXCheckBox.IsValid() && SendDMXCheckBox->IsChecked())
+	if (WeakFaderEntity.IsValid() && CachedProtocol.IsValid()
+		&& SendDMXCheckBox.IsValid() && SendDMXCheckBox->IsChecked()
+		&& OnValueChanged.IsBound())
 	{
-		for (const TSharedPtr<SDMXFaderChannel>& Channel : FaderChannels)
-		{
-			if (WeakFaderEntity.Get())
-			{
-				IDMXProtocolPtr DMXProtocol = WeakFaderEntity->DeviceProtocol;
-				if (DMXProtocol.IsValid())
-				{
-					IDMXFragmentMap DMXFragmentMap;
-					DMXFragmentMap.Add(Channel->GetChannelNumber(), CurrentFaderValue);
-					DMXProtocol->SendDMXFragment(Channel->GetUniverseNumber(), DMXFragmentMap);
-				}
-			}
-		}
+		OnValueChanged.Execute(SharedThis(this));
 	}
 }
-
 FText SDMXFader::GetProtocolText() const
 {
 	if (WeakFaderEntity.IsValid())
 	{
-		return FText::FromName(WeakFaderEntity->DeviceProtocol);
+		return FText::FromName(CachedProtocol);
 	}
 
 	return FText::GetEmpty();
diff --git a/Engine/Plugins/Experimental/DigitalMultiplex/DMXEngine/Source/DMXEditor/Private/Widgets/OutputFader/SDMXFader.h b/Engine/Plugins/Experimental/DigitalMultiplex/DMXEngine/Source/DMXEditor/Private/Widgets/OutputFader/SDMXFader.h
--- a/Engine/Plugins/Experimental/DigitalMultiplex/DMXEngine/Source/DMXEditor/Private/Widgets/OutputFader/SDMXFader.h
+++ b/Engine/Plugins/Experimental/DigitalMultiplex/DMXEngine/Source/DMXEditor/Private/Widgets/OutputFader/SDMXFader.h
@@ -8,6 +8,7 @@
 #include "Styling/SlateTypes.h"
 
 #include "DMXProtocolCommon.h"
+#include "DMXProtocolTypes.h"
 
 template<typename NumericType>
 class SSpinBoxVertical;
@@ -20,6 +21,7 @@ class STextBlock;
 class SDMXFaderChannel;
 class SCheckBox;
 class IDMXProtocol;
+struct FDMXUniverse;
 
 #define LOCTEXT_NAMESPACE "SDMXFader"
 
@@ -28,13 +30,25 @@ class SDMXFader
 	: public SCompoundWidget
 {
 public:
+
+	/**
+	 * Notification for fader value change
+	 * Parameters: Protocol, Universe ID, Channel Address, New Fader Value
+	 */
+	DECLARE_DELEGATE_OneParam(FOnFaderChanged, TSharedRef<SDMXFader>);
+
 	SLATE_BEGIN_ARGS(SDMXFader)
 		: _InText(LOCTEXT("FaderLabel", "Fader"))
+		, _OnValueChanged()
+		, _OnSendStateChanged()
 	{}
 
 	SLATE_ARGUMENT(TWeakPtr<FDMXEditor>, DMXEditor)
 	SLATE_ATTRIBUTE(FText, InText)
 
+	SLATE_EVENT(FOnFaderChanged, OnValueChanged)
+	SLATE_EVENT(FOnFaderChanged, OnSendStateChanged)
+
 	SLATE_END_ARGS()
 
 	/** Constructs the widget */
@@ -46,7 +60,7 @@ public:
 		
 	void SetFaderLabel(const FString& InLabel);
 
-	void AddChannelWidget(const FString& InUniverse, const FString& InChannel, uint16 InUniverseNumber, uint32 InChannelNumber, const IDMXProtocolPtr& DMXProtocol);
+	void AddChannelWidget(const FString& InUniverse, const FString& InChannel, uint16 InUniverseNumber, uint32 InChannelNumber);
 
 	void RemoveAllChannelWidgets();
 
@@ -58,6 +72,17 @@ public:
 
 	const TSharedPtr<SSpinBoxVertical<uint8>>& GetFaderBoxVertical() const { return FaderBoxVertical; }
 
+	/** Returns a list of the Universes/Channels this fader widget is using */
+	const TArray<TSharedPtr<SDMXFaderChannel>>& GetChannels() const { return FaderChannels; }
+
+	/** Current fader value */
+	uint8 GetCurrentValue() const { return CurrentFaderValue; }
+
+	bool ShouldSendDMX() const;
+
+	void SetProtocol(const FDMXProtocolName& InProtocol) { CachedProtocol = InProtocol; }
+	FDMXProtocolName GetProtocol() const { return CachedProtocol; }
+
 private:
 	//~ Begin SWidget implementation
 	virtual FReply OnMouseButtonDown(const FGeometry& MyGeometry, const FPointerEvent& MouseEvent) override;
@@ -65,6 +90,8 @@ private:
 
 	FReply HandleRemoveFaderClicked();
 
+	void HandleSendDMXCheckChanged(ECheckBoxState NewState);
+
 	void UpdateFaderTemplateProperties();
 
 	/**Change fader background color on hover */
@@ -93,7 +120,7 @@ private:
 	/** Pointer back to the DMXEditor tool that owns us */
 	TWeakPtr<FDMXEditor> WeakDMXEditor;
 
-	/** Pointer back to assosiated fader entity UObject */
+	/** Pointer back to associated fader entity UObject */
 	TWeakObjectPtr<UDMXEntityFader> WeakFaderEntity;
 
 	/** Pointer back to fader list widget */
@@ -102,6 +129,11 @@ private:
 	TArray<TSharedPtr<SDMXFaderChannel>> FaderChannels;
 
 	TSharedPtr<SCheckBox> SendDMXCheckBox;
+
+	FDMXProtocolName CachedProtocol;
+
+	FOnFaderChanged OnValueChanged;
+	FOnFaderChanged OnSendStateChanged;
 };
 
 #undef LOCTEXT_NAMESPACE
\ No newline at end of file
diff --git a/Engine/Plugins/Experimental/DigitalMultiplex/DMXEngine/Source/DMXEditor/Private/Widgets/OutputFader/SDMXFaderChannel.cpp b/Engine/Plugins/Experimental/DigitalMultiplex/DMXEngine/Source/DMXEditor/Private/Widgets/OutputFader/SDMXFaderChannel.cpp
--- a/Engine/Plugins/Experimental/DigitalMultiplex/DMXEngine/Source/DMXEditor/Private/Widgets/OutputFader/SDMXFaderChannel.cpp
+++ b/Engine/Plugins/Experimental/DigitalMultiplex/DMXEngine/Source/DMXEditor/Private/Widgets/OutputFader/SDMXFaderChannel.cpp
@@ -22,7 +22,7 @@ void SDMXFaderChannel::Construct(const FArguments& InArgs)
 	WeakDMXEditor = InArgs._DMXEditor;
 	UniverseNumber = InArgs._UniverseNumber;
 	ChannelNumber = InArgs._ChannelNumber;
-	WeakDMXProtocol = InArgs._DMXProtocol;
+
 
 	ChildSlot
 	.Padding(0.f)
diff --git a/Engine/Plugins/Experimental/DigitalMultiplex/DMXEngine/Source/DMXEditor/Private/Widgets/OutputFader/SDMXFaderChannel.h b/Engine/Plugins/Experimental/DigitalMultiplex/DMXEngine/Source/DMXEditor/Private/Widgets/OutputFader/SDMXFaderChannel.h
--- a/Engine/Plugins/Experimental/DigitalMultiplex/DMXEngine/Source/DMXEditor/Private/Widgets/OutputFader/SDMXFaderChannel.h
+++ b/Engine/Plugins/Experimental/DigitalMultiplex/DMXEngine/Source/DMXEditor/Private/Widgets/OutputFader/SDMXFaderChannel.h
@@ -9,7 +9,6 @@
 
 class STextBlock;
 class FDMXEditor;
-class IDMXProtocol;
 
 #define LOCTEXT_NAMESPACE "SDMXFaderChannel"
 
@@ -27,7 +26,6 @@ public:
 	SLATE_ARGUMENT(FText, InText)
 	SLATE_ARGUMENT(uint16, UniverseNumber)
 	SLATE_ARGUMENT(uint32, ChannelNumber)
-	SLATE_ARGUMENT(IDMXProtocolPtrWeak, DMXProtocol)
 
 	SLATE_END_ARGS()
 
@@ -42,8 +40,6 @@ public:
 
 	uint32 GetChannelNumber() const { return ChannelNumber; }
 
-	const IDMXProtocolPtr GetProtocol() const { return WeakDMXProtocol.Pin(); }
-
 private:
 	/** Pointer back to the DMXEditor tool that owns us */
 	TWeakPtr<FDMXEditor> WeakDMXEditor;
@@ -55,8 +51,6 @@ private:
 	uint16 UniverseNumber;
 
 	uint32 ChannelNumber;
-
-	IDMXProtocolPtrWeak WeakDMXProtocol;
 };
 
 #undef LOCTEXT_NAMESPACE
\ No newline at end of file
diff --git a/Engine/Plugins/Experimental/DigitalMultiplex/DMXEngine/Source/DMXEditor/Private/Widgets/OutputFader/SDMXOutputFaderList.cpp b/Engine/Plugins/Experimental/DigitalMultiplex/DMXEngine/Source/DMXEditor/Private/Widgets/OutputFader/SDMXOutputFaderList.cpp
--- a/Engine/Plugins/Experimental/DigitalMultiplex/DMXEngine/Source/DMXEditor/Private/Widgets/OutputFader/SDMXOutputFaderList.cpp
+++ b/Engine/Plugins/Experimental/DigitalMultiplex/DMXEngine/Source/DMXEditor/Private/Widgets/OutputFader/SDMXOutputFaderList.cpp
@@ -11,6 +11,7 @@
 #include "DMXEditorUtils.h"
 
 #include "Widgets/OutputFader/SDMXFader.h"
+#include "Widgets/OutputFader/SDMXFaderChannel.h"
 
 #include "Widgets/Input/SButton.h"
 #include "Widgets/Layout/SScrollBox.h"
@@ -114,12 +115,133 @@ FReply SDMXOutputFaderList::HandleUpdateFaderClicked()
 {
 	if (TSharedPtr<SDMXFader> SelectedPtr = WeakSelectedFaderWidget.Pin())
 	{
+		UDMXEntityFader* FaderTemplate = WeakFaderTemplate.Get();
+		check(FaderTemplate);
+
+		// Is the Fader protocol still the same?
+		if (FaderTemplate->DeviceProtocol == SelectedPtr->GetProtocol())
+		{
+			// Check for Universes/Channels about to be removed from the fader
+			const TArray<TSharedPtr<SDMXFaderChannel>>& FaderWidgetChannels = SelectedPtr->GetChannels();
+			const TArray<FDMXUniverse>& FaderEntityUniverses = FaderTemplate->Universes;
+
+			for (const TSharedPtr<SDMXFaderChannel>& WidgetChannel : FaderWidgetChannels)
+			{
+				bool bHasUniverseAndChannel = false;
+				for (const FDMXUniverse& EntityUniverse : FaderEntityUniverses)
+				{
+					if (WidgetChannel->GetUniverseNumber() == EntityUniverse.UniverseNumber
+						&& WidgetChannel->GetChannelNumber() == EntityUniverse.Channel)
+					{
+						bHasUniverseAndChannel = true;
+						break;
+					}
+				}
+
+				// Has this universe/channel been deleted?
+				if (!bHasUniverseAndChannel)
+				{
+					CompactFragmentMaps(WidgetChannel->GetUniverseNumber(), WidgetChannel->GetChannelNumber(), SelectedPtr.ToSharedRef());
+				}
+			}
+		}
+		else // protocol has been changed
+		{
+			// Try to delete FragmentMaps related to this fader, in case they are not
+			// used by any other faders.
+			for (const TSharedPtr<SDMXFaderChannel>& FaderChannel : SelectedPtr->GetChannels())
+			{
+				CompactFragmentMaps(FaderChannel->GetUniverseNumber(), FaderChannel->GetChannelNumber(), SelectedPtr.ToSharedRef());
+			}
+		}
+
+		check(SelectedPtr->GetWeakFaderEntity().IsValid());
+		UDMXEntityFader* FaderEntity = SelectedPtr->GetWeakFaderEntity().Get();
+
 		SelectedPtr->RemoveAllChannelWidgets();
-		SetFaderProperties(SelectedPtr, SelectedPtr->GetWeakFaderEntity(), true);
+		SetFaderProperties(SelectedPtr, FaderEntity, true);
+
+		if (SelectedPtr->ShouldSendDMX())
+		{
+			// Make sure the entries in FragmentMaps related to this fader's Universes and Addresses all exist
+			HandleFaderSendStateChanged(SelectedPtr.ToSharedRef());
+			// Sends data from this fader to keep the output up to date with the new changes
+			HandleFaderValueChanged(SelectedPtr.ToSharedRef());
+		}
 	}
+
 	return FReply::Handled();
 }
 
+void SDMXOutputFaderList::HandleFaderValueChanged(TSharedRef<SDMXFader> InFaderWidget)
+{
+	const IDMXProtocolPtr Protocol = InFaderWidget->GetProtocol();
+	check(Protocol.IsValid());
+
+	const FName ProtocolName = Protocol->GetProtocolName();
+	const TArray<TSharedPtr<SDMXFaderChannel>>& FaderChannels = InFaderWidget->GetChannels();
+	const uint8 NewValue = InFaderWidget->GetCurrentValue();
+
+	TSet<uint16> ChangedUniverseIds;
+	ChangedUniverseIds.Reserve(FaderChannels.Num());
+
+	// Update all necessary fragments affected by this fader
+	for (const TSharedPtr<SDMXFaderChannel>& FaderChannel : FaderChannels)
+	{
+		FragmentMaps[ProtocolName][FaderChannel->GetUniverseNumber()][FaderChannel->GetChannelNumber()] = NewValue;
+		ChangedUniverseIds.Add(FaderChannel->GetUniverseNumber());
+	}
+
+	// Sends fragment maps for each of the affected universes
+	for (const uint16& UniverseID : ChangedUniverseIds)
+	{
+		Protocol->SendDMXFragment(UniverseID, FragmentMaps[ProtocolName][UniverseID]);
+	}
+}
+
+void SDMXOutputFaderList::HandleFaderSendStateChanged(TSharedRef<SDMXFader> InFaderWidget)
+{
+	const bool bSendDMXData = InFaderWidget->ShouldSendDMX();
+
+	if (bSendDMXData)
+	{
+		check(InFaderWidget->GetProtocol());
+		const FName ProtocolName = InFaderWidget->GetProtocol();
+
+		if (FragmentMaps.Find(ProtocolName) == nullptr)
+		{
+			FragmentMaps.Add(ProtocolName);
+		}
+
+		// Make sure the fragment maps for the universes and addresses related to this fader
+		// exist in the FragmentMaps entries.
+		for (const TSharedPtr<SDMXFaderChannel>& FaderChannel : InFaderWidget->GetChannels())
+		{
+			const uint16 UniverseID = FaderChannel->GetUniverseNumber();
+			const uint16 Address = FaderChannel->GetChannelNumber();
+
+			if (FragmentMaps[ProtocolName].Find(UniverseID) == nullptr)
+			{
+				FragmentMaps[ProtocolName].Add(UniverseID);
+			}
+
+			if (FragmentMaps[ProtocolName][UniverseID].Find(Address) == nullptr)
+			{
+				FragmentMaps[ProtocolName][UniverseID].Add(Address);
+			}
+		}
+	}
+	else
+	{
+		// Try to delete FragmentMaps related to this fader, in case they are not
+		// used by any other faders.
+		for (const TSharedPtr<SDMXFaderChannel>& FaderChannel : InFaderWidget->GetChannels())
+		{
+			CompactFragmentMaps(FaderChannel->GetUniverseNumber(), FaderChannel->GetChannelNumber(), InFaderWidget);
+		}
+	}
+}
+
 void SDMXOutputFaderList::ResetFaderBackgrounds()
 {
 	for (TSharedPtr<SDMXFader> Fader : FaderWidgets)
@@ -133,13 +255,16 @@ void SDMXOutputFaderList::ResetFaderBackgrounds()
 
 void SDMXOutputFaderList::AddFader(const FString& InName /*= TEXT("")*/)
 {
-	TSharedPtr<SDMXFader> NewFader = SNew(SDMXFader)
+	TSharedRef<SDMXFader> NewFader = SNew(SDMXFader)
 		.DMXEditor(WeakDMXEditor)
-		.InText(LOCTEXT("FaderLabel", "Fader"));
-		FaderSlots->AddSlot()
-		[
-			NewFader.ToSharedRef()
-		];
+		.InText(LOCTEXT("FaderLabel", "Fader"))
+		.OnValueChanged(this, &SDMXOutputFaderList::HandleFaderValueChanged)
+		.OnSendStateChanged(this, &SDMXOutputFaderList::HandleFaderSendStateChanged);
+
+	FaderSlots->AddSlot()
+	[
+		NewFader
+	];
 
 	if (TSharedPtr<FDMXEditor> DMXEditorPtr = WeakDMXEditor.Pin())
 	{
@@ -147,8 +272,6 @@ void SDMXOutputFaderList::AddFader(const FString& InName /*= TEXT("")*/)
 		{
 			if (UDMXEntityFader* FaderEntity = Cast<UDMXEntityFader>(DMXLibrary->GetOrCreateEntityObject(InName, UDMXEntityFader::StaticClass())))
 			{
-				
-
 				NewFader->SetFaderEntity(FaderEntity);
 				NewFader->SetParentFaderList(SharedThis(this));
 				FaderWidgets.Add(NewFader);
@@ -156,11 +279,101 @@ void SDMXOutputFaderList::AddFader(const FString& InName /*= TEXT("")*/)
 				SetFaderProperties(NewFader, FaderEntity, InName.IsEmpty());
 				DeselectFaders();
 				NewFader->SelectThisFader();
+
+				// Make sure the fader's channels have entries related to them in FragmentMaps
+				HandleFaderSendStateChanged(NewFader);
 			}
 		}
 	}
 }
 
+void SDMXOutputFaderList::CompactFragmentMaps(uint16 RemovedUniverseID, uint16 RemovedAddress, const TSharedRef<SDMXFader> FaderInstigator)
+{
+	bool bUniverseIsStillUsed = false;
+	bool bAddressIsStillUsed = false;
+
+	const FDMXProtocolName RemoveFromProtocol = FaderInstigator->GetProtocol();
+
+	// If the universe or address entries already don't exist in FragmentMaps, we don't need to do anything.
+	if (FragmentMaps.Find(RemoveFromProtocol) == nullptr
+		|| FragmentMaps[RemoveFromProtocol].Find(RemovedUniverseID) == nullptr
+		|| FragmentMaps[RemoveFromProtocol][RemovedUniverseID].Find(RemovedAddress) == nullptr)
+	{
+		return;
+	}
+
+	for (const TSharedPtr<SDMXFader>& FaderWidget : FaderWidgets)
+	{
+		// Skip the fader which channels are being removed from. We only need to know about the other faders.
+		if (FaderWidget == FaderInstigator)
+		{
+			continue;
+		}
+
+		// Skip faders that aren't supposed to send data
+		if (!FaderWidget->ShouldSendDMX())
+		{
+			continue;
+		}
+
+		// Skip faders with different protocols. They might have the Universe IDs we're looking for,
+		// but we're only looking for Universe IDs on the Instigator Fader protocol.
+		if (FaderWidget->GetProtocol() != RemoveFromProtocol)
+		{
+			continue;
+		}
+
+		const TArray<TSharedPtr<SDMXFaderChannel>>& FaderChannels = FaderWidget->GetChannels();
+		for (const TSharedPtr<SDMXFaderChannel>& FaderChannel : FaderChannels)
+		{
+			if (FaderChannel->GetUniverseNumber() == RemovedUniverseID)
+			{
+				// This universe can't be removed. It's used on other faders.
+				bUniverseIsStillUsed = true;
+
+				if (FaderChannel->GetChannelNumber() == RemovedAddress)
+				{
+					// This Address, in this universe, also can't be removed.
+					bAddressIsStillUsed = true;
+
+					// Now we know both can't be removed. There's no need to keep searching through this fader's channels.
+					break;
+				}
+			}
+		}
+
+		if (bAddressIsStillUsed)
+		{
+			// If we already know the address is still used, which means its universe is as well,
+			// there's no need to keep searching through other faders.
+			break;
+		}
+	}
+
+	if (!bUniverseIsStillUsed)
+	{
+		// Deletes the whole FragmentMap for the unused Universe ID
+		FragmentMaps[RemoveFromProtocol].Remove(RemovedUniverseID);
+
+		UE_LOG_DMXEDITOR(Log, TEXT("%S: Removed unused Universe ID %d from Output"), __FUNCTION__, RemovedUniverseID);
+	}
+	else if (!bAddressIsStillUsed)
+	{
+		// Deletes just the specific address value from the FragmentMap
+		FragmentMaps[RemoveFromProtocol][RemovedUniverseID].Remove(RemovedAddress);
+		// Delete the Universe in case this was the last registered address in it
+		if (FragmentMaps[RemoveFromProtocol][RemovedUniverseID].Num() == 0)
+		{
+			FragmentMaps[RemoveFromProtocol].Remove(RemovedUniverseID);
+			UE_LOG_DMXEDITOR(Log, TEXT("%S: Removed unused Universe ID %d from Output"), __FUNCTION__, RemovedUniverseID);
+		}
+		else
+		{
+			UE_LOG_DMXEDITOR(Log, TEXT("%S: Removed unused address %d from Universe ID %d from Output"), __FUNCTION__, RemovedAddress, RemovedUniverseID);
+		}
+	}
+}
+
 void SDMXOutputFaderList::RemoveFader(TSharedPtr<SDMXFader> FaderToRemove)
 {
 	if (!FaderToRemove.IsValid())
@@ -174,7 +387,17 @@ void SDMXOutputFaderList::RemoveFader(TSharedPtr<SDMXFader> FaderToRemove)
 	{
 		if (UDMXLibrary* DMXLibrary = DMXEditorPtr->GetDMXLibrary())
 		{
-			DMXLibrary->RemoveEntity(FaderToRemove->GetWeakFaderEntity()->GetDisplayName());
+			if (UDMXEntityFader* FaderEntity = FaderToRemove->GetWeakFaderEntity().Get())
+			{
+				// Try to delete FragmentMaps related to this fader, in case they are not
+				// used anymore by any other faders.
+				for (const FDMXUniverse& Universe : FaderEntity->Universes)
+				{
+					CompactFragmentMaps(Universe.UniverseNumber, Universe.Channel, FaderToRemove.ToSharedRef());
+				}
+
+				DMXLibrary->RemoveEntity(FaderEntity);
+			}
 		}
 	}
 	FaderSlots->RemoveSlot(FaderToRemove.ToSharedRef());
@@ -200,15 +423,15 @@ void SDMXOutputFaderList::SetFaderProperties(const TSharedPtr<SDMXFader>& InFade
 		TransferCreatedFaderObjectProperties(InFaderObject, WeakFaderTemplate);
 	}
 
-	IDMXProtocolPtr DMXProtocol = InFaderObject->DeviceProtocol;
-
+	InFaderWidget->SetProtocol(InFaderObject->DeviceProtocol);
 	InFaderWidget->SetFaderLabel(InFaderObject->GetDisplayName());
+
 	for (const FDMXUniverse& Universe : InFaderObject->Universes)
 	{
 		// Add protocol universe
 		FString TmpFinalUniverse = FString::FromInt(Universe.UniverseNumber);
 		FString TmpChannel = FString::FromInt(Universe.Channel);
-		InFaderWidget->AddChannelWidget(TmpFinalUniverse, TmpChannel, Universe.UniverseNumber, Universe.Channel, DMXProtocol);
+		InFaderWidget->AddChannelWidget(TmpFinalUniverse, TmpChannel, Universe.UniverseNumber, Universe.Channel);
 	}
 }
 
diff --git a/Engine/Plugins/Experimental/DigitalMultiplex/DMXEngine/Source/DMXEditor/Private/Widgets/OutputFader/SDMXOutputFaderList.h b/Engine/Plugins/Experimental/DigitalMultiplex/DMXEngine/Source/DMXEditor/Private/Widgets/OutputFader/SDMXOutputFaderList.h
--- a/Engine/Plugins/Experimental/DigitalMultiplex/DMXEngine/Source/DMXEditor/Private/Widgets/OutputFader/SDMXOutputFaderList.h
+++ b/Engine/Plugins/Experimental/DigitalMultiplex/DMXEngine/Source/DMXEditor/Private/Widgets/OutputFader/SDMXOutputFaderList.h
@@ -5,6 +5,7 @@
 #include "CoreMinimal.h"
 #include "Widgets/DeclarativeSyntaxSupport.h"
 #include "Widgets/SCompoundWidget.h"
+#include "DMXProtocolCommon.h"
 
 class SDMXFader;
 class FDMXEditor;
@@ -48,6 +49,18 @@ private:
 
 	FReply HandleUpdateFaderClicked();
 
+	/**
+	 * Called when a fader value changes.
+	 * Updates the value in the correct Fragment Map and sends it over.
+	 */
+	void HandleFaderValueChanged(TSharedRef<SDMXFader> InFaderWidget);
+
+	/**
+	 * Called when a fader send DMX check box state changes.
+	 * It creates or removes entries in the FragmentMaps entries depending on the Fader's Send state.
+	 */
+	void HandleFaderSendStateChanged(TSharedRef<SDMXFader> InFaderWidget);
+
 	void SetFaderProperties(const TSharedPtr<SDMXFader>& InFaderWidget, const TWeakObjectPtr<UDMXEntityFader>& InFaderObject, bool bIsTransferObject);
 
 	void TransferCreatedFaderObjectProperties(const TWeakObjectPtr<UDMXEntityFader>& TransferTo, const TWeakObjectPtr<UDMXEntityFader>& TransferFrom);
@@ -55,6 +68,13 @@ private:
 
 	void AddFader(const FString& InName = TEXT(""));
 
+	/**
+	 * Called when a Fader's Channel or Universe is removed to try and free memory by deleting
+	 * FragmentMaps or addresses that are not in use anymore.
+	 * It prevents values not being controlled by faders anymore from being sent.
+	 */
+	void CompactFragmentMaps(uint16 RemovedUniverseID, uint16 RemovedAddress, const TSharedRef<SDMXFader> FaderInstigator);
+
 public:
 	TWeakPtr<SDMXFader> WeakSelectedFaderWidget;
 
@@ -69,4 +89,10 @@ private:
 	TWeakObjectPtr<UDMXEntityFader> WeakFaderTemplate;
 
 	TWeakObjectPtr<UDMXLibrary> WeakDMXLibrary;
+
+	/**
+	 * Keeps fragment maps for each universe in each protocol.
+	 * <Protocol Name => <Universe ID => Fragment Map> >
+	 */
+	TMap<FName, TMap<uint16, IDMXFragmentMap> > FragmentMaps;
 };
diff --git a/Engine/Plugins/Experimental/DigitalMultiplex/DMXEngine/Source/DMXRuntime/Public/Library/DMXEntityController.h b/Engine/Plugins/Experimental/DigitalMultiplex/DMXEngine/Source/DMXRuntime/Public/Library/DMXEntityController.h
--- a/Engine/Plugins/Experimental/DigitalMultiplex/DMXEngine/Source/DMXRuntime/Public/Library/DMXEntityController.h
+++ b/Engine/Plugins/Experimental/DigitalMultiplex/DMXEngine/Source/DMXRuntime/Public/Library/DMXEntityController.h
@@ -15,6 +15,9 @@ public:
 	/**  Defines where DMX data is sent to. */
 	UPROPERTY(EditAnywhere, Category = "DMX") // Hidden. For now it has just a single value. More values will be added in the future.
 	EDMXCommunicationTypes CommunicationMode;
+	
+	UPROPERTY(EditAnywhere, Category = "DMX")
+	FString UnicastIP;
 
 	/**  First Universe ID on this Controller's range. */
 	UPROPERTY(EditAnywhere, BlueprintReadOnly, Category = "Universe Properties", meta = (DisplayName = "Universe Start", DisplayPriority = 1, ClampMin = 1))
@@ -55,7 +58,9 @@ public:
 
 public:
 	UDMXEntityController()
-		: UniverseLocalNum(1)
+		: UnicastIP(TEXT("0.0.0.0"))
+		, UniverseLocalStart(1)
+		, UniverseLocalNum(1)
 	{}
 
 	/** Returns the currently assigned protocol for this controller */
diff --git a/Engine/Plugins/Experimental/DigitalMultiplex/DMXProtocol/Source/DMXProtocol/Public/DMXProtocolConstants.h b/Engine/Plugins/Experimental/DigitalMultiplex/DMXProtocol/Source/DMXProtocol/Public/DMXProtocolConstants.h
--- a/Engine/Plugins/Experimental/DigitalMultiplex/DMXProtocol/Source/DMXProtocol/Public/DMXProtocolConstants.h
+++ b/Engine/Plugins/Experimental/DigitalMultiplex/DMXProtocol/Source/DMXProtocol/Public/DMXProtocolConstants.h
@@ -32,6 +32,8 @@ enum EDMUniverseDirection
 #define DMX_MAX_FINALUNIVERSE 32768
 #define DMX_MAX_ADDRESS 512
 #define DMX_MAX_VALUE 255
+#define DMX_MAX_PACKET_SIZE 2048u
 #define DMX_MAX_FUNCTION_SIZE ((int32)sizeof(int32))
 
+
 #define DMX_K2_CATEGORY_NAME "DMX"
\ No newline at end of file
diff --git a/Engine/Plugins/Experimental/DigitalMultiplex/DMXProtocol/Source/DMXProtocolSACN/Private/DMXProtocolTransportSACN.cpp b/Engine/Plugins/Experimental/DigitalMultiplex/DMXProtocol/Source/DMXProtocolSACN/Private/DMXProtocolTransportSACN.cpp
--- a/Engine/Plugins/Experimental/DigitalMultiplex/DMXProtocol/Source/DMXProtocolSACN/Private/DMXProtocolTransportSACN.cpp
+++ b/Engine/Plugins/Experimental/DigitalMultiplex/DMXProtocol/Source/DMXProtocolSACN/Private/DMXProtocolTransportSACN.cpp
@@ -2,14 +2,11 @@
 
 #include "DMXProtocolTransportSACN.h"
 #include "DMXProtocolSACN.h"
+#include "DMXProtocolTypes.h"
 
-#include "Common/UdpSocketSender.h"
 #include "HAL/Event.h"
 #include "HAL/Runnable.h"
 #include "HAL/RunnableThread.h"
-#include "IMessageAttachment.h"
-#include "Serialization/ArrayReader.h"
-#include "Serialization/ArrayWriter.h"
 #include "Sockets.h"
 #include "SocketSubsystem.h"
 
@@ -24,6 +21,9 @@ FDMXProtocolSenderSACN::FDMXProtocolSenderSACN(FSocket& InSocket, FDMXProtocolSA
 	, BroadcastSocket(&InSocket)
 	, Protocol(InProtocol)
 {
+	SocketSubsystem = ISocketSubsystem::Get(PLATFORM_SOCKETSUBSYSTEM);
+	check(SocketSubsystem);
+
 	WorkEvent = MakeShareable(FPlatformProcess::GetSynchEventFromPool(), [](FEvent* EventToDelete)
 	{
 		FPlatformProcess::ReturnSynchEventToPool(EventToDelete);
@@ -105,8 +105,6 @@ bool FDMXProtocolSenderSACN::EnqueueOutboundPackage(FDMXPacketPtr Packet)
 
 void FDMXProtocolSenderSACN::ConsumeOutboundPackages()
 {
-	FScopeLock ScopeLock(&CriticalSection);
-
 	FDMXPacketPtr Packet;
 	while (OutboundPackages.Dequeue(Packet))
 	{
@@ -115,99 +113,10 @@ void FDMXProtocolSenderSACN::ConsumeOutboundPackages()
 		bool bIsSent = BroadcastSocket->SendTo(Packet->Data.GetData(), Packet->Data.Num(), BytesSent, *FDMXProtocolSACN::GetUniverseAddr(Packet->Settings.GetNumberField(TEXT("UniverseID"))));
 		if (!bIsSent)
 		{
-			UE_LOG_DMXPROTOCOL(Error, TEXT("Error sending %d"), bIsSent);
-		}
-	}
-}
-
-FDMXProtocolReceiverSACN::FDMXProtocolReceiverSACN(FSocket& InSocket, const FTimespan& InWaitTime)
-	: Socket(&InSocket)
-	, Stopping(false)
-	, Thread(nullptr)
-	, WaitTime(InWaitTime)
-{
-	check(Socket != nullptr);
-	check(Socket->GetSocketType() == SOCKTYPE_Datagram);
-
-	SocketSubsystem = ISocketSubsystem::Get(PLATFORM_SOCKETSUBSYSTEM);
-
-	Thread = FRunnableThread::Create(this, *ThreadName, 128 * 1024, TPri_Normal, FPlatformAffinity::GetPoolThreadMask());
-}
-
-FDMXProtocolReceiverSACN::~FDMXProtocolReceiverSACN()
-{
-	if (Thread != nullptr)
-	{
-		Thread->Kill(true);
-		delete Thread;
-	}
-}
-
-FOnDMXDataReceived& FDMXProtocolReceiverSACN::OnDataReceived()
-{
-	return DMXDataReceiveDelegate;
-}
-
-FRunnableThread* FDMXProtocolReceiverSACN::GetThread() const
-{
-	return Thread;
-}
-
-bool FDMXProtocolReceiverSACN::Init()
-{
-	return true;
-}
+			ESocketErrors RecvFromError = SocketSubsystem->GetLastErrorCode();
 
-uint32 FDMXProtocolReceiverSACN::Run()
-{
-	while (!Stopping)
-	{
-		Update(WaitTime);
-	}
-
-	return 0;
-}
-
-void FDMXProtocolReceiverSACN::Stop()
-{
-	Stopping = true;
-}
-
-void FDMXProtocolReceiverSACN::Exit()
-{
-}
-
-void FDMXProtocolReceiverSACN::Tick()
-{
-	Update(FTimespan::Zero());
-}
-
-FSingleThreadRunnable* FDMXProtocolReceiverSACN::GetSingleThreadInterface()
-{
-	return this;
-}
-
-void FDMXProtocolReceiverSACN::Update(const FTimespan& SocketWaitTime)
-{
-	if (!Socket->Wait(ESocketWaitConditions::WaitForRead, SocketWaitTime))
-	{
-		return;
-	}
-
-	TSharedRef<FInternetAddr> Sender = SocketSubsystem->CreateInternetAddr();
-	uint32 Size;
-
-	while (Socket->HasPendingData(Size))
-	{
-		FArrayReaderPtr Reader = MakeShared<FArrayReader, ESPMode::ThreadSafe>(true);
-		Reader->SetNumUninitialized(FMath::Min(Size, 65507u));
-
-		int32 Read = 0;
-
-		if (Socket->RecvFrom(Reader->GetData(), Reader->Num(), Read, *Sender))
-		{
-			Reader->RemoveAt(Read, Reader->Num() - Read, false);
-			DMXDataReceiveDelegate.ExecuteIfBound(Reader);
+			UE_LOG_DMXPROTOCOL(Error, TEXT("Error sending %d"),
+				(uint8)RecvFromError);
 		}
 	}
 }
diff --git a/Engine/Plugins/Experimental/DigitalMultiplex/DMXProtocol/Source/DMXProtocolSACN/Private/DMXProtocolTransportSACN.h b/Engine/Plugins/Experimental/DigitalMultiplex/DMXProtocol/Source/DMXProtocolSACN/Private/DMXProtocolTransportSACN.h
--- a/Engine/Plugins/Experimental/DigitalMultiplex/DMXProtocol/Source/DMXProtocolSACN/Private/DMXProtocolTransportSACN.h
+++ b/Engine/Plugins/Experimental/DigitalMultiplex/DMXProtocol/Source/DMXProtocolSACN/Private/DMXProtocolTransportSACN.h
@@ -4,11 +4,10 @@
 
 #include "CoreMinimal.h"
 
-#include "DMXProtocolTypes.h"
+#include "DMXProtocolCommon.h"
 #include "Interfaces/IDMXProtocolTransport.h"
 
 #include "Containers/Queue.h"
-#include "HAL/CriticalSection.h"
 
 class FSocket;
 class ISocketSubsystem;
@@ -50,8 +49,6 @@ private:
 	/** Holds the last sent message number. */
 	int32 LastSentPackage;
 
-	mutable FCriticalSection CriticalSection;
-
 	FThreadSafeCounter StopTaskCounter;
 
 	/** Holds the thread object. */
@@ -66,59 +63,6 @@ private:
 	FSocket* BroadcastSocket;
 
 	FDMXProtocolSACN* Protocol;
-};
-
 
-class FDMXProtocolReceiverSACN
-	: public IDMXProtocolReceiver
-{
-public:
-	FDMXProtocolReceiverSACN(FSocket& InSocket, const FTimespan& InWaitTime);
-	virtual ~FDMXProtocolReceiverSACN();
-
-public:
-	//~ Begin FRunnable implementation
-	virtual bool Init() override;
-	virtual uint32 Run() override;
-	virtual void Stop() override;
-	virtual void Exit() override;
-	//~ End FRunnable implementation
-
-	//~ Begin FSingleThreadRunnable implementation
-	virtual void Tick() override;
-	virtual class FSingleThreadRunnable* GetSingleThreadInterface() override;
-	//~ End FSingleThreadRunnable implementation
-
-	//~ Begin IDMXProtocolReceiver implementation
-	virtual FOnDMXDataReceived& OnDataReceived() override;
-	virtual FRunnableThread* GetThread() const override;
-	//~ End IDMXProtocolReceiver implementation
-
-protected:
-	void Update(const FTimespan& SocketWaitTime);
-
-private:
-
-	/** The network socket. */
-	FSocket* Socket;
-
-	/** Pointer to the socket sub-system. */
 	ISocketSubsystem* SocketSubsystem;
-
-	/** Flag indicating that the thread is stopping. */
-	bool Stopping;
-
-	/** The thread object. */
-	FRunnableThread* Thread;
-
-	/** The receiver thread's name. */
-	FString ThreadName;
-
-	/** The amount of time to wait for inbound packets. */
-	FTimespan WaitTime;
-
-private:
-
-	/** Holds the data received delegate. */
-	FOnDMXDataReceived DMXDataReceiveDelegate;
-};
\ No newline at end of file
+};
diff --git a/Engine/Plugins/Experimental/DigitalMultiplex/DMXProtocol/Source/DMXProtocolSACN/Private/DMXProtocolUniverseSACN.cpp b/Engine/Plugins/Experimental/DigitalMultiplex/DMXProtocol/Source/DMXProtocolSACN/Private/DMXProtocolUniverseSACN.cpp
--- a/Engine/Plugins/Experimental/DigitalMultiplex/DMXProtocol/Source/DMXProtocolSACN/Private/DMXProtocolUniverseSACN.cpp
+++ b/Engine/Plugins/Experimental/DigitalMultiplex/DMXProtocol/Source/DMXProtocolSACN/Private/DMXProtocolUniverseSACN.cpp
@@ -5,28 +5,43 @@
 #include "DMXProtocolSACN.h"
 #include "DMXProtocolConstants.h"
 #include "DMXProtocolTypes.h"
-#include "DMXProtocolTransportSACN.h"
 #include "DMXProtocolSettings.h"
+#include "DMXProtocolSACNUtils.h"
 
 #include "Common/UdpSocketBuilder.h"
+#include "IPAddress.h"
+#include "Serialization/ArrayReader.h"
+#include "SocketSubsystem.h"
 
-#include "DMXProtocolSACNUtils.h"
+const double FDMXProtocolUniverseSACN::TimeWithoutInputBufferRequest = 5.0;
 
 FDMXProtocolUniverseSACN::FDMXProtocolUniverseSACN(IDMXProtocolPtr InDMXProtocol, const FJsonObject& InSettings)
 	: WeakDMXProtocol(InDMXProtocol)
 	, ListeningSocket(nullptr)
+	, NetworkErrorMessagePrefix(TEXT("NETWORK ERROR SACN:"))
 {
-	NetworkErrorMessagePrefix = TEXT("NETWORK ERROR SACN:");
+	// Not ticking by default
+	TimeWithoutInputBufferRequestEnd = TimeWithoutInputBufferRequestStart = FPlatformTime::Seconds();
+	bIsTicking = false;
 
-	Settings = MakeShared<FJsonObject>(InSettings);
+	// Set online subsustem
+	SocketSubsystem = ISocketSubsystem::Get(PLATFORM_SOCKETSUBSYSTEM);
+	check(SocketSubsystem);
+
+	// Set listener address
+	ListenerInternetAddr = SocketSubsystem->CreateInternetAddr();
 
+	// Sets the settings
+	Settings = MakeShared<FJsonObject>(InSettings);
 	checkf(WeakDMXProtocol.IsValid(), TEXT("DMXProtocol pointer is not valid"));
 	checkf(Settings->HasField(TEXT("UniverseID")), TEXT("DMXProtocol UniverseID is not valid"));
 	UniverseID = Settings->GetNumberField(TEXT("UniverseID"));
 
+	// Allocate new buffers
 	OutputDMXBuffer = MakeShared<FDMXBuffer>();
 	InputDMXBuffer = MakeShared<FDMXBuffer>();
 
+	// Set default IP address
 	InterfaceIPAddress = GetDefault<UDMXProtocolSettings>()->InterfaceIPAddress;
 
 	// Set Network Interface listener
@@ -58,6 +73,10 @@ TSharedPtr<FDMXBuffer> FDMXProtocolUniverseSACN::GetOutputDMXBuffer() const
 
 TSharedPtr<FDMXBuffer> FDMXProtocolUniverseSACN::GetInputDMXBuffer() const
 {
+	TimeWithoutInputBufferRequestStart = FPlatformTime::Seconds();
+	TimeWithoutInputBufferRequestEnd = TimeWithoutInputBufferRequestStart + TimeWithoutInputBufferRequest;
+	bIsTicking = true;
+
 	return InputDMXBuffer;
 }
 
@@ -86,57 +105,71 @@ bool FDMXProtocolUniverseSACN::IsSupportRDM() const
 	return true;
 }
 
-void FDMXProtocolUniverseSACN::OnDataReceived(const FArrayReaderPtr & Buffer)
+bool FDMXProtocolUniverseSACN::ReceiveDMXBuffer()
 {
-	// It will be more handlers
-	switch (SACN::GetRootPacketType(Buffer))
+	if (FSocket* Socket = GetOrCreateListeningSocket())
 	{
-	case VECTOR_ROOT_E131_DATA:
-		HandleReplyPacket(Buffer);
-		break;
-	default:
-		break;
-	}
-}
-
-bool FDMXProtocolUniverseSACN::HandleReplyPacket(const FArrayReaderPtr & Buffer)
-{
-	// SACN PDU packets
-	*Buffer << IncomingDMXRootLayer;
-	*Buffer << IncomingDMXFramingLayer;
-	*Buffer << IncomingDMXDMPLayer;
+		uint32 Size = 0;
+		int32 Read = 0;
+		FArrayReaderPtr Reader = MakeShared<FArrayReader, ESPMode::ThreadSafe>(true);
 
-	bool bCopySuccessful = false;
-	// Access the buffer thread-safely
-	InputDMXBuffer->AccessDMXData([this, &bCopySuccessful](TArray<uint8>& InData)
+		// Atempt to read if there is any incoming data
+		while (Socket->HasPendingData(Size))
 		{
-			// Make sure we copy same amount of data
-			if (InData.Num() == ACN_DMX_SIZE)
+			Reader->SetNumUninitialized(FMath::Min(Size, DMX_MAX_PACKET_SIZE));
+
+			// Read buffer from socket
+			if (Socket->RecvFrom(Reader->GetData(), Reader->Num(), Read, *ListenerInternetAddr))
 			{
-				InputDMXBuffer->SetDMXBuffer(IncomingDMXDMPLayer.DMX, ACN_DMX_SIZE);
-				IDMXProtocolPtr Protocol = GetProtocol();
-				if (Protocol.IsValid())
-				{
-					Protocol->GetOnUniverseInputUpdate().Broadcast(Protocol->GetProtocolName(), GetUniverseID(), InData);
-					bCopySuccessful = true;
-				}
-				else
-				{
-					UE_LOG_DMXPROTOCOL(Error, TEXT("Invalid SACN Protocol"));
-					bCopySuccessful = false;
-				}
+				Reader->RemoveAt(Read, Reader->Num() - Read, false);
+
+				return OnDataReceived(Reader);
 			}
 			else
 			{
-				UE_LOG_DMXPROTOCOL(Error, TEXT("%s: Size of incoming DMX buffer is wrong! Expected: %d; Found: %d")
-					, NetworkErrorMessagePrefix
-					, ACN_DMX_SIZE
-					, InData.Num());
-				bCopySuccessful = false;
+				ESocketErrors RecvFromError = SocketSubsystem->GetLastErrorCode();
+
+				UE_LOG_DMXPROTOCOL(Error, TEXT("Error recieving the packet Universe ID %d, error is %d"), 
+					GetUniverseID(),
+					(uint8)RecvFromError);
 			}
-		});
+		}
+	}
+	else
+	{
+		UE_LOG_DMXPROTOCOL(Error, TEXT("Error recieving. No valid socket for universe %d"), GetUniverseID());
+	}
 
-	return bCopySuccessful;
+
+	return false;
+}
+
+void FDMXProtocolUniverseSACN::Tick(float DeltaTime)
+{
+	if (TimeWithoutInputBufferRequestStart > TimeWithoutInputBufferRequestEnd)
+	{
+		// Stop listening and destroy the socket
+		ReleaseNetworkInterface();
+
+		bIsTicking = false;
+	}
+	else
+	{
+		// Keep listening
+		ReceiveDMXBuffer();
+
+		TimeWithoutInputBufferRequestStart = FPlatformTime::Seconds();
+	}
+}
+
+bool FDMXProtocolUniverseSACN::IsTickable() const
+{
+	return bIsTicking;
+}
+
+TStatId FDMXProtocolUniverseSACN::GetStatId() const
+{
+	RETURN_QUICK_DECLARE_CYCLE_STAT(FDMXProtocolUniverseSACN, STATGROUP_Tickables);
 }
 
 void FDMXProtocolUniverseSACN::OnNetworkInterfaceChanged(const FString& InInterfaceIPAddress)
@@ -150,32 +183,103 @@ void FDMXProtocolUniverseSACN::OnNetworkInterfaceChanged(const FString& InInterf
 
 bool FDMXProtocolUniverseSACN::RestartNetworkInterface(const FString& InInterfaceIPAddress, FString& OutErrorMessage)
 {
-	FScopeLock Lock(&ListeningSocketsCS);
-
 	// Clean the error message
 	OutErrorMessage.Empty();
 
+	// Release old network interface
+	ReleaseNetworkInterface();
+
+	// Set new network interface IP
+	InterfaceIPAddress = InInterfaceIPAddress;
+
+	return true;
+}
+
+void FDMXProtocolUniverseSACN::ReleaseNetworkInterface()
+{
+	// Clean all sockets
+	if (ListeningSocket != nullptr)
+	{
+		ListeningSocket->Close();
+		SocketSubsystem->DestroySocket(ListeningSocket);
+		ListeningSocket = nullptr;
+	}
+}
+
+void FDMXProtocolUniverseSACN::SetLayerPackets(const FArrayReaderPtr& Buffer)
+{
+	*Buffer << IncomingDMXRootLayer;
+	*Buffer << IncomingDMXFramingLayer;
+	*Buffer << IncomingDMXDMPLayer;
+}
+
+bool FDMXProtocolUniverseSACN::OnDataReceived(const FArrayReaderPtr& Buffer)
+{
+	// It will be more handlers
+	switch (SACN::GetRootPacketType(Buffer))
+	{
+	case VECTOR_ROOT_E131_DATA:
+		return HandleReplyPacket(Buffer);
+	default:
+		return false;
+	}
+}
+
+bool FDMXProtocolUniverseSACN::HandleReplyPacket(const FArrayReaderPtr& Buffer)
+{
+	bool bCopySuccessful = false;
+
+	// Copy the data from incoming socket buffer to SACN universe
+	SetLayerPackets(Buffer);
+
+	// Access the buffer thread-safety
+	InputDMXBuffer->AccessDMXData([this, &bCopySuccessful](TArray<uint8>& InData)
+	{
+		// Make sure we copy same amount of data
+		if (InData.Num() == ACN_DMX_SIZE)
+		{
+			InputDMXBuffer->SetDMXBuffer(IncomingDMXDMPLayer.DMX, ACN_DMX_SIZE);
+
+			GetProtocol()->GetOnUniverseInputUpdate().Broadcast(GetProtocol()->GetProtocolName(), UniverseID, InData);
+			bCopySuccessful = true;
+		}
+		else
+		{
+			UE_LOG_DMXPROTOCOL(Error, TEXT("%s: Size of incoming DMX buffer is wrong! Expected: %d; Found: %d")
+				, NetworkErrorMessagePrefix
+				, ACN_DMX_SIZE
+				, InData.Num());
+			bCopySuccessful = false;
+		}
+	});
+
+	return bCopySuccessful;
+}
+
+FSocket* FDMXProtocolUniverseSACN::GetOrCreateListeningSocket()
+{
+	// Simple return the socket if it exists
+	if (ListeningSocket != nullptr)
+	{
+		return ListeningSocket;
+	}
+
 	// Try to create IP address at the first
-	ISocketSubsystem* SocketSubsystem = ISocketSubsystem::Get(PLATFORM_SOCKETSUBSYSTEM);
-	TSharedPtr<FInternetAddr> SenderAddr = SocketSubsystem->CreateInternetAddr();
+	TSharedPtr<FInternetAddr> InterfaceAddr = SocketSubsystem->CreateInternetAddr();
 	bool bIsValid = false;
-	SenderAddr->SetIp(*InInterfaceIPAddress, bIsValid);
-    //SenderAddr->SetAnyAddress();
-	SenderAddr->SetPort(ACN_PORT);
+	InterfaceAddr->SetIp(*InterfaceIPAddress, bIsValid);
+	InterfaceAddr->SetPort(ACN_PORT);
 	if (!bIsValid)
 	{
-		OutErrorMessage = FString::Printf(TEXT("Wrong IP address: %s"), *InInterfaceIPAddress);
-		return false;
+		UE_LOG_DMXPROTOCOL(Error, TEXT("Wrong IP address: %s"), *InterfaceIPAddress);
+		return nullptr;
 	}
 
-	// Release old network interface
-	ReleaseNetworkInterface();
-
-	FIPv4Endpoint SenderEndpoint = FIPv4Endpoint(SenderAddr);
+	FIPv4Endpoint InterfaceEndpoint = FIPv4Endpoint(InterfaceAddr);
 	TSharedPtr<FInternetAddr> ListeningAddr = FDMXProtocolSACN::GetUniverseAddr(UniverseID);
 	FIPv4Endpoint ListeningEndpoint = FIPv4Endpoint(ListeningAddr);
 
-	FSocket* NewListeningSocket = FUdpSocketBuilder(TEXT("SACNListeningSocket"))
+	FSocket* NewListeningSocket = FUdpSocketBuilder(FString::Printf(TEXT("SACNListeningSocket_Universe_%d"), UniverseID))
 		.AsNonBlocking()
 
 #if PLATFORM_WINDOWS
@@ -184,53 +288,26 @@ bool FDMXProtocolUniverseSACN::RestartNetworkInterface(const FString& InInterfac
 		// On all other platforms we bind to the wildcard IP address in order
 		// to be able to also receive packets that were sent directly to the
 		// interface IP instead of the multicast address.
-		.BoundToAddress(SenderEndpoint.Address)
+		.BoundToAddress(InterfaceEndpoint.Address)
 #endif
-		.BoundToPort(SenderEndpoint.Port)
+		.BoundToPort(InterfaceEndpoint.Port)
 #if PLATFORM_SUPPORTS_UDP_MULTICAST_GROUP
-		.JoinedToGroup(ListeningEndpoint.Address, SenderEndpoint.Address)
+		.JoinedToGroup(ListeningEndpoint.Address, InterfaceEndpoint.Address)
 		.WithMulticastLoopback()
 		.WithMulticastTtl(1)
-		.WithMulticastInterface(SenderEndpoint.Address)
+		.WithMulticastInterface(InterfaceEndpoint.Address)
 #endif
-        .AsReusable();
-
+		.AsReusable()
+		.WithReceiveBufferSize(ACN_DMX_SIZE * 4);
 
 	if (NewListeningSocket == nullptr)
 	{
-		OutErrorMessage = FString::Printf(TEXT("Error create ListeningSocket: %s"), *InterfaceIPAddress);
-		return false;
+		UE_LOG_DMXPROTOCOL(Error, TEXT("Error create ListeningSocket: %s"), *InterfaceIPAddress);
+		return nullptr;
 	}
 
-	// Set new network interface IP
-	InterfaceIPAddress = InInterfaceIPAddress;
-
 	// Save New socket;
 	ListeningSocket = NewListeningSocket;
 
-	// Set new receiver
-	SACNReceiver = MakeShared<FDMXProtocolReceiverSACN>(*ListeningSocket, FTimespan::FromMilliseconds(100));
-	SACNReceiver->OnDataReceived().BindRaw(this, &FDMXProtocolUniverseSACN::OnDataReceived);
-
-	return true;
-}
-
-void FDMXProtocolUniverseSACN::ReleaseNetworkInterface()
-{
-	if (SACNReceiver.IsValid())
-	{
-		SACNReceiver.Reset();
-		SACNReceiver = nullptr;
-	}
-
-	// Clean all sockets
-	if (ListeningSocket != nullptr)
-	{
-		ListeningSocket->Close();
-		if (ISocketSubsystem* SocketSubsystem = ISocketSubsystem::Get(PLATFORM_SOCKETSUBSYSTEM))
-		{
-			SocketSubsystem->DestroySocket(ListeningSocket);
-		}
-		ListeningSocket = nullptr;
-	}
+	return ListeningSocket;
 }
diff --git a/Engine/Plugins/Experimental/DigitalMultiplex/DMXProtocol/Source/DMXProtocolSACN/Private/DMXProtocolUniverseSACN.h b/Engine/Plugins/Experimental/DigitalMultiplex/DMXProtocol/Source/DMXProtocolSACN/Private/DMXProtocolUniverseSACN.h
--- a/Engine/Plugins/Experimental/DigitalMultiplex/DMXProtocol/Source/DMXProtocolSACN/Private/DMXProtocolUniverseSACN.h
+++ b/Engine/Plugins/Experimental/DigitalMultiplex/DMXProtocol/Source/DMXProtocolSACN/Private/DMXProtocolUniverseSACN.h
@@ -4,18 +4,25 @@
 
 #include "CoreMinimal.h"
 
-#include "Interfaces/IPv4/IPv4Endpoint.h"
+#include "DMXProtocolCommon.h"
+#include "Packets/DMXProtocolE131PDUPacket.h"
 #include "Interfaces/IDMXProtocolUniverse.h"
 #include "Interfaces/IDMXNetworkInterface.h"
-#include "Packets/DMXProtocolE131PDUPacket.h"
 
-#include "HAL/CriticalSection.h"
+#include "Tickable.h"
 
+class FInternetAddr;
 class FSocket;
+class ISocketSubsystem;
 
+/**
+ * SACN universe implementation
+ * It holds the input and output buffer
+ */
 class DMXPROTOCOLSACN_API FDMXProtocolUniverseSACN
 	: public IDMXProtocolUniverse
 	, public IDMXNetworkInterface
+	, public FTickableGameObject
 {
 public:
 	FDMXProtocolUniverseSACN(IDMXProtocolPtr InDMXProtocol, const FJsonObject& InSettings);
@@ -32,6 +39,14 @@ public:
 	virtual bool IsSupportRDM() const override;
 	//~ End IDMXProtocolDevice implementation
 
+	//~ Begin FTickableObjectBase implementation
+	virtual void Tick(float DeltaTime) override;
+	virtual bool IsTickable() const;
+	virtual TStatId GetStatId() const override;
+	virtual bool IsTickableInEditor() const override { return true; }
+	virtual bool IsTickableWhenPaused() const override { return true; }
+	//~ End FTickableObjectBase implementation
+
 	//~ Begin IDMXNetworkInterface implementation
 	virtual void OnNetworkInterfaceChanged(const FString& InInterfaceIPAddress) override;
 	virtual bool RestartNetworkInterface(const FString& InInterfaceIPAddress, FString& OutErrorMessage) override;
@@ -44,9 +59,20 @@ public:
 	const FDMXProtocolE131DMPLayerPacket& GetIncomingDMXDMPLayer() const { return IncomingDMXDMPLayer; }
 
 private:
-	void OnDataReceived(const FArrayReaderPtr & Buffer);
+	/** Called when new DMX packet is coming */
+	bool OnDataReceived(const FArrayReaderPtr& Buffer);
+
+	/** Handle and broadcasting incoming DMX packet */
+	bool HandleReplyPacket(const FArrayReaderPtr& Buffer);
 
-	bool HandleReplyPacket(const FArrayReaderPtr & Buffer);
+	/** Ask to read DMX from the Socket */
+	bool ReceiveDMXBuffer();
+
+	/** Parse incoming DMX packet into the SACN layers */
+	void SetLayerPackets(const FArrayReaderPtr& Buffer);
+
+	/** Checks if the socket exists, otherwise create a new one */
+	FSocket* GetOrCreateListeningSocket();
 
 private:
 	IDMXProtocolPtrWeak WeakDMXProtocol;
@@ -57,16 +83,37 @@ private:
 	bool bIsRDMSupport;
 	TSharedPtr<FJsonObject> Settings;
 
+	/** The universe listening socket */
 	FSocket* ListeningSocket;
-	TSharedPtr<IDMXProtocolReceiver> SACNReceiver;
+
+	/** The universe socket listening Addr */
+	TSharedPtr<FInternetAddr> ListenerInternetAddr;
+
+	/**
+	 * Current number of ticks without DMX input buffer request.
+	 * It the amount as same as MaxNumTicks the socket will be destroyed, which will save the resources.
+	 * If (NumTicksWithoutInputBufferRequest == 0) there is no socket reading
+	 * If (NumTicksWithoutInputBufferRequest > 0 && NumTicksWithoutInputBufferRequest < MaxNumTicksWithoutInputBufferRequest)
+	 *   the socket is reading each tick, if it new read request the new socket will be created.
+	 * If (NumTicksWithoutInputBufferRequest == MaxNumTicksWithoutInputBufferRequest) there is no socket will be destoyed
+	 */
+	mutable double TimeWithoutInputBufferRequestStart;
+
+	/** Max number of ticks without DMX input buffer request. The socket will be destroyed after this point */
+	mutable double TimeWithoutInputBufferRequestEnd;
+
+	mutable bool bIsTicking;
+
+	static const double TimeWithoutInputBufferRequest;
+
+
+	/** Pointer to the socket sub-system. */
+	ISocketSubsystem* SocketSubsystem;
 
 	FDMXProtocolE131RootLayerPacket IncomingDMXRootLayer;
 	FDMXProtocolE131FramingLayerPacket IncomingDMXFramingLayer;
 	FDMXProtocolE131DMPLayerPacket IncomingDMXDMPLayer;
 
-	/** Mutex protecting access to the sockets. */
-	mutable FCriticalSection ListeningSocketsCS;
-
 	FString InterfaceIPAddress;
 
 	FDelegateHandle NetworkInterfaceChangedHandle;
diff --git a/Engine/Plugins/Experimental/DigitalMultiplex/DMXProtocol/Source/DMXProtocolSACN/Private/Tests/DMXPrtotocolSACNFlowTest.cpp b/Engine/Plugins/Experimental/DigitalMultiplex/DMXProtocol/Source/DMXProtocolSACN/Private/Tests/DMXPrtotocolSACNFlowTest.cpp
--- a/Engine/Plugins/Experimental/DigitalMultiplex/DMXProtocol/Source/DMXProtocolSACN/Private/Tests/DMXPrtotocolSACNFlowTest.cpp
+++ b/Engine/Plugins/Experimental/DigitalMultiplex/DMXProtocol/Source/DMXProtocolSACN/Private/Tests/DMXPrtotocolSACNFlowTest.cpp
@@ -3,6 +3,7 @@
 #include "CoreMinimal.h"
 #include "Engine/Engine.h"
 #include "Misc/AutomationTest.h"
+#include "Tests/AutomationCommon.h"
 #include "DMXProtocolSACNModule.h"
 
 #include "DMXProtocolSACN.h"
@@ -10,107 +11,190 @@
 #include "DMXProtocolSACN.h"
 #include "DMXProtocolUniverseSACN.h"
 
-static uint16 UniverseValue = 34;
+struct DMXPrtotocolSACNHelper
+{
+	DMXPrtotocolSACNHelper(uint16 InUniverseID, FAutomationTestBase* InTest)
+		: UniverseID(InUniverseID)
+		, Test(InTest)
+	{
+		// Init Protocol and Managers
+		DMXProtocol = static_cast<FDMXProtocolSACN*>((IDMXProtocol::Get(FDMXProtocolSACNModule::NAME_SACN)).Get());
 
-IMPLEMENT_SIMPLE_AUTOMATION_TEST(DMXPrtotocolSACNBasicFlowTest, "VirtualProduction.DMX.SACN.BasicFlow", EAutomationTestFlags::ApplicationContextMask | EAutomationTestFlags::EngineFilter)
+		// Init Universe
+		FJsonObject UniverseSettings;
+		UniverseSettings.SetNumberField(TEXT("UniverseID"), InUniverseID);
+		Universe = DMXProtocol->AddUniverse(UniverseSettings);
 
-bool DMXPrtotocolSACNBasicFlowTest::RunTest(const FString& Parameters)
+		// Call get buffer for start listening socket
+		Universe->GetInputDMXBuffer();
+	}
+
+	~DMXPrtotocolSACNHelper()
+	{
+		DMXProtocol->RemoveUniverseById(UniverseID);
+	}
+
+	FDMXProtocolSACN* DMXProtocol;
+	IDMXProtocolUniversePtr Universe;
+	uint16 UniverseID;
+	uint8 FixtureChannels[6] = { 1, 2, 3, 4, 5, 6 };
+	uint8 FixtureValues[6] = { 255, 155, 50, 100, 200, 220 };
+
+	/** Pointer to running automation test instance */
+	FAutomationTestBase* Test;
+};
+
+DEFINE_LATENT_AUTOMATION_COMMAND_ONE_PARAMETER(DMXPrtotocolSACNBasicFlow_CheckUniverse_Part1, TSharedPtr<DMXPrtotocolSACNHelper>, Helper);
+bool DMXPrtotocolSACNBasicFlow_CheckUniverse_Part1::Update()
 {
-	// Init Protocol and Managers
-	FDMXProtocolSACN* DMXProtocol = static_cast<FDMXProtocolSACN*>((IDMXProtocol::Get(FDMXProtocolSACNModule::NAME_SACN)).Get());
+	TSharedPtr<IDMXProtocolUniverse, ESPMode::ThreadSafe> Universe = Helper->DMXProtocol->GetUniverseById(Helper->UniverseID);
 
-	FJsonObject UniverseSettings;
-	UniverseSettings.SetNumberField(TEXT("UniverseID"), UniverseValue);
-	DMXProtocol->AddUniverse(UniverseSettings);
-	TSharedPtr<IDMXProtocolUniverse, ESPMode::ThreadSafe> Universe = DMXProtocol->GetUniverseById(UniverseValue);
+	Helper->Test->TestTrue(TEXT("Universe is valid"), Universe.IsValid());
+	Helper->Test->TestEqual(TEXT("Universe ID should be equal"), Universe->GetUniverseID(), Helper->UniverseID);
 
-	TestTrue(TEXT("Universe is valid"), Universe.IsValid());
-	TestEqual(TEXT("Universe ID should be equal"), Universe->GetUniverseID(), UniverseValue);
+	return true;
+}
 
-	// Send DMX fragment 1
-	const uint8 FixtureChannels[6] = { 1, 2, 3, 4, 5, 6 };
-	const uint8 FixtureValues[6] = { 255, 155, 50, 100, 200, 220 };
+DEFINE_LATENT_AUTOMATION_COMMAND_ONE_PARAMETER(DMXPrtotocolSACNBasicFlow_SendDMX1_Part2, TSharedPtr<DMXPrtotocolSACNHelper>, Helper);
+bool DMXPrtotocolSACNBasicFlow_SendDMX1_Part2::Update()
+{
 	IDMXFragmentMap DMXFragmentMap;
-	DMXFragmentMap.Add(FixtureChannels[0], FixtureValues[0]);
-	DMXFragmentMap.Add(FixtureChannels[1], FixtureValues[1]);
-	DMXFragmentMap.Add(FixtureChannels[2], FixtureValues[2]);
-	DMXProtocol->SendDMXFragment(Universe->GetUniverseID(), DMXFragmentMap);
+	DMXFragmentMap.Add(Helper->FixtureChannels[0], Helper->FixtureValues[0]);
+	DMXFragmentMap.Add(Helper->FixtureChannels[1], Helper->FixtureValues[1]);
+	DMXFragmentMap.Add(Helper->FixtureChannels[2], Helper->FixtureValues[2]);
+
+	Helper->DMXProtocol->SendDMXFragment(Helper->Universe->GetUniverseID(), DMXFragmentMap);
+
+	return true;
+}
 
-	// Wait before receive DMX
-	FPlatformProcess::Sleep(0.2f);
-	TSharedPtr<FDMXBuffer> InputDMXBuffer = Universe->GetInputDMXBuffer();
+DEFINE_LATENT_AUTOMATION_COMMAND_ONE_PARAMETER(DMXPrtotocolSACNBasicFlow_CheckDMX1_Part3, TSharedPtr<DMXPrtotocolSACNHelper>, Helper);
+bool DMXPrtotocolSACNBasicFlow_CheckDMX1_Part3::Update()
+{
+	TSharedPtr<FDMXBuffer> InputDMXBuffer = Helper->Universe->GetInputDMXBuffer();
 
 	// Test package ACNPacketIdentifier
 	{
 		const uint8 ACNPacketIdentifier[ACN_IDENTIFIER_SIZE] = { 'A', 'S', 'C', '-', 'E', '1', '.', '1', '7', '\0', '\0', '\0' };
-		FDMXProtocolUniverseSACN* DMXProtocolUniverseSACN = static_cast<FDMXProtocolUniverseSACN*>(Universe.Get());
+		FDMXProtocolUniverseSACN* DMXProtocolUniverseSACN = static_cast<FDMXProtocolUniverseSACN*>(Helper->Universe.Get());
 		FString ExpectedInentifier = FString::Printf(TEXT("%s"), ACNPacketIdentifier);
 		FString IncomingInentifier = FString::Printf(TEXT("%s"), DMXProtocolUniverseSACN->GetIncomingDMXRootLayer().ACNPacketIdentifier);
-		TestTrue(TEXT("Incoming ACNPacketIdentifier should be the same"), ExpectedInentifier.Equals(IncomingInentifier));
+		Helper->Test->TestTrue(TEXT("Incoming ACNPacketIdentifier should be the same"), ExpectedInentifier.Equals(IncomingInentifier));
 	}
 
-	TestEqual(TEXT("Incoming buffer should be same"), FixtureValues[0], InputDMXBuffer->GetDMXDataAddress(FixtureChannels[0] - 1));
-	TestEqual(TEXT("Incoming buffer should be same"), FixtureValues[1], InputDMXBuffer->GetDMXDataAddress(FixtureChannels[1] - 1));
-	TestEqual(TEXT("Incoming buffer should be same"), FixtureValues[2], InputDMXBuffer->GetDMXDataAddress(FixtureChannels[2] - 1));
+	Helper->Test->TestEqual(TEXT("Incoming buffer should be same"), Helper->FixtureValues[0], InputDMXBuffer->GetDMXDataAddress(Helper->FixtureChannels[0] - 1));
+	Helper->Test->TestEqual(TEXT("Incoming buffer should be same"), Helper->FixtureValues[1], InputDMXBuffer->GetDMXDataAddress(Helper->FixtureChannels[1] - 1));
+	Helper->Test->TestEqual(TEXT("Incoming buffer should be same"), Helper->FixtureValues[2], InputDMXBuffer->GetDMXDataAddress(Helper->FixtureChannels[2] - 1));
+
+	return true;
+}
+
+DEFINE_LATENT_AUTOMATION_COMMAND_ONE_PARAMETER(DMXPrtotocolSACNBasicFlow_SendDMX2_Part4, TSharedPtr<DMXPrtotocolSACNHelper>, Helper);
+bool DMXPrtotocolSACNBasicFlow_SendDMX2_Part4::Update()
+{
+	TSharedPtr<FDMXBuffer> InputDMXBuffer = Helper->Universe->GetInputDMXBuffer();
+
+	IDMXFragmentMap DMXFragmentMap;
+
+	DMXFragmentMap.Add(Helper->FixtureChannels[3], Helper->FixtureValues[3]);
+	DMXFragmentMap.Add(Helper->FixtureChannels[4], Helper->FixtureValues[4]);
+	DMXFragmentMap.Add(Helper->FixtureChannels[5], Helper->FixtureValues[5]);
+	Helper->DMXProtocol->SendDMXFragment(Helper->Universe->GetUniverseID(), DMXFragmentMap);
+
+	return true;
+}
+
+DEFINE_LATENT_AUTOMATION_COMMAND_ONE_PARAMETER(DMXPrtotocolSACNBasicFlow_CheckDMX2_Part5, TSharedPtr<DMXPrtotocolSACNHelper>, Helper);
+bool DMXPrtotocolSACNBasicFlow_CheckDMX2_Part5::Update()
+{
+	TSharedPtr<FDMXBuffer> InputDMXBuffer = Helper->Universe->GetInputDMXBuffer();
+
+	Helper->Test->TestEqual(TEXT("Old value DMX input should be the same"), Helper->FixtureValues[0], InputDMXBuffer->GetDMXDataAddress(Helper->FixtureChannels[0] - 1));
+	Helper->Test->TestEqual(TEXT("Incoming buffer should be same"), Helper->FixtureValues[4], InputDMXBuffer->GetDMXDataAddress(Helper->FixtureChannels[4] - 1));
+	Helper->Test->TestEqual(TEXT("Incoming buffer should be same"), Helper->FixtureValues[5], InputDMXBuffer->GetDMXDataAddress(Helper->FixtureChannels[5] - 1));
 
-	// Send DMX fragment 2
-	DMXFragmentMap.Add(FixtureChannels[3], FixtureValues[3]);
-	DMXFragmentMap.Add(FixtureChannels[4], FixtureValues[4]);
-	DMXFragmentMap.Add(FixtureChannels[5], FixtureValues[5]);
-	DMXProtocol->SendDMXFragment(Universe->GetUniverseID(), DMXFragmentMap);
-	FPlatformProcess::Sleep(0.2f);
+	return true;
+}
 
-	TestEqual(TEXT("Old value DMX input should be the same"), FixtureValues[0], InputDMXBuffer->GetDMXDataAddress(FixtureChannels[0] - 1));
-	TestEqual(TEXT("Incoming buffer should be same"), FixtureValues[4], InputDMXBuffer->GetDMXDataAddress(FixtureChannels[4] - 1));
-	TestEqual(TEXT("Incoming buffer should be same"), FixtureValues[5], InputDMXBuffer->GetDMXDataAddress(FixtureChannels[5] - 1));
 
-	// @TODO Implement test cases for Discovery
+DEFINE_LATENT_AUTOMATION_COMMAND_ONE_PARAMETER(DMXPrtotocolSACNBasicFlow_SendDiscovery_Part6, TSharedPtr<DMXPrtotocolSACNHelper>, Helper);
+bool DMXPrtotocolSACNBasicFlow_SendDiscovery_Part6::Update()
+{
 	TArray<uint16> Universes;
 	for (int32 UniversesIndex = 0; UniversesIndex < DMX_UNIVERSE_SIZE; UniversesIndex++)
 	{
 		Universes.Add(UniversesIndex);
 	}
-	DMXProtocol->SendDiscovery(Universes);
+	Helper->DMXProtocol->SendDiscovery(Universes);
 
 	return true;
 }
 
-IMPLEMENT_SIMPLE_AUTOMATION_TEST(FDMXPrtotocolSACNConsoleCommandsTest, "VirtualProduction.DMX.SACN.ConsoleCommands", EAutomationTestFlags::ApplicationContextMask | EAutomationTestFlags::EngineFilter)
-
-bool FDMXPrtotocolSACNConsoleCommandsTest::RunTest(const FString& Parameters)
+IMPLEMENT_SIMPLE_AUTOMATION_TEST(DMXPrtotocolSACNBasicFlowTest, "VirtualProduction.DMX.SACN.BasicFlow", EAutomationTestFlags::ApplicationContextMask | EAutomationTestFlags::EngineFilter)
+bool DMXPrtotocolSACNBasicFlowTest::RunTest(const FString& Parameters)
 {
-	// Init Protocol and Managers
-	FDMXProtocolSACN* DMXProtocol = static_cast<FDMXProtocolSACN*>((IDMXProtocol::Get(FDMXProtocolSACNModule::NAME_SACN)).Get());
+	TSharedPtr<DMXPrtotocolSACNHelper> Helper = MakeShared<DMXPrtotocolSACNHelper>(34, this);
 
-	FJsonObject UniverseSettings;
-	UniverseSettings.SetNumberField(TEXT("UniverseID"), UniverseValue);
-	DMXProtocol->AddUniverse(UniverseSettings);
-	TSharedPtr<IDMXProtocolUniverse, ESPMode::ThreadSafe> Universe = DMXProtocol->GetUniverseById(UniverseValue);
+	// Wait the universe socket and buffer setup
+	ADD_LATENT_AUTOMATION_COMMAND(FWaitLatentCommand(0.2f));
 
-	// DMX.SACN.SendDMX[UniverseID] Channel:Value Channel : Value Channel : Value ...
-	uint32 Values[3][2] =
-	{
-		{95, 105},
-		{88, 41},
-		{12, 1},
-	};
+	ADD_LATENT_AUTOMATION_COMMAND(DMXPrtotocolSACNBasicFlow_CheckUniverse_Part1(Helper));
+	ADD_LATENT_AUTOMATION_COMMAND(DMXPrtotocolSACNBasicFlow_SendDMX1_Part2(Helper));
+
+	// Wait after sending first DMX
+	ADD_LATENT_AUTOMATION_COMMAND(FWaitLatentCommand(0.2f));
+
+	ADD_LATENT_AUTOMATION_COMMAND(DMXPrtotocolSACNBasicFlow_CheckDMX1_Part3(Helper));
+	ADD_LATENT_AUTOMATION_COMMAND(DMXPrtotocolSACNBasicFlow_SendDMX2_Part4(Helper));
+
+	// Wait after sending second DMX
+	ADD_LATENT_AUTOMATION_COMMAND(FWaitLatentCommand(0.2f));
+
+	ADD_LATENT_AUTOMATION_COMMAND(DMXPrtotocolSACNBasicFlow_CheckDMX2_Part5(Helper));
+	ADD_LATENT_AUTOMATION_COMMAND(DMXPrtotocolSACNBasicFlow_SendDiscovery_Part6(Helper));
+
+	return true;
+}
 
+DEFINE_LATENT_AUTOMATION_COMMAND_ONE_PARAMETER(FDMXPrtotocolSACNConsoleCommands_SendDMX1_Part1, TSharedPtr<DMXPrtotocolSACNHelper>, Helper);
+bool FDMXPrtotocolSACNConsoleCommands_SendDMX1_Part1::Update()
+{
+	// DMX.SACN.SendDMX[UniverseID] Channel:Value Channel : Value Channel : Value ...
 	FString Command(
 		FString::Printf(TEXT("DMX.SACN.SendDMX %d %d:%d %d:%d %d:%d"),
-			Universe->GetUniverseID(),
-			Values[0][0], Values[0][1],
-			Values[1][0], Values[1][1],
-			Values[2][0], Values[2][1]
+			Helper->Universe->GetUniverseID(),
+			Helper->FixtureChannels[0], Helper->FixtureValues[0],
+			Helper->FixtureChannels[1], Helper->FixtureValues[1],
+			Helper->FixtureChannels[2], Helper->FixtureValues[2]
 		));
 
 	GEngine->Exec(GWorld, *Command);
 
-	FPlatformProcess::Sleep(0.2f);
+	return true;
+}
+
+DEFINE_LATENT_AUTOMATION_COMMAND_ONE_PARAMETER(FDMXPrtotocolSACNConsoleCommands_CheckDMX1_Part2, TSharedPtr<DMXPrtotocolSACNHelper>, Helper);
+bool FDMXPrtotocolSACNConsoleCommands_CheckDMX1_Part2::Update()
+{
+	TSharedPtr<FDMXBuffer> InputDMXBuffer = Helper->Universe->GetInputDMXBuffer();
+	Helper->Test->TestEqual(TEXT("Incoming buffer should be same"), Helper->FixtureValues[0], InputDMXBuffer->GetDMXDataAddress(Helper->FixtureChannels[0] - 1));
+	Helper->Test->TestEqual(TEXT("Incoming buffer should be same"), Helper->FixtureValues[1], InputDMXBuffer->GetDMXDataAddress(Helper->FixtureChannels[1] - 1));
+	Helper->Test->TestEqual(TEXT("Incoming buffer should be same"), Helper->FixtureValues[2], InputDMXBuffer->GetDMXDataAddress(Helper->FixtureChannels[2] - 1));
+
+	return true;
+}
+
+IMPLEMENT_SIMPLE_AUTOMATION_TEST(FDMXPrtotocolSACNConsoleCommandsTest, "VirtualProduction.DMX.SACN.ConsoleCommands", EAutomationTestFlags::ApplicationContextMask | EAutomationTestFlags::EngineFilter)
+bool FDMXPrtotocolSACNConsoleCommandsTest::RunTest(const FString& Parameters)
+{
+	TSharedPtr<DMXPrtotocolSACNHelper> Helper = MakeShared<DMXPrtotocolSACNHelper>(576, this);
+
+	// Wait the universe socket and buffer setup
+	ADD_LATENT_AUTOMATION_COMMAND(FWaitLatentCommand(0.2f));
+	ADD_LATENT_AUTOMATION_COMMAND(FDMXPrtotocolSACNConsoleCommands_SendDMX1_Part1(Helper));
 
-	TSharedPtr<FDMXBuffer> InputDMXBuffer = Universe->GetInputDMXBuffer();
-	TestEqual(TEXT("Incoming buffer should be same"), Values[0][1], InputDMXBuffer->GetDMXDataAddress(Values[0][0] - 1));
-	TestEqual(TEXT("Incoming buffer should be same"), Values[1][1], InputDMXBuffer->GetDMXDataAddress(Values[1][0] - 1));
-	TestEqual(TEXT("Incoming buffer should be same"), Values[2][1], InputDMXBuffer->GetDMXDataAddress(Values[2][0] - 1));
+	// Wait after sending first DMX
+	ADD_LATENT_AUTOMATION_COMMAND(FWaitLatentCommand(0.2f));
 
 	return true;
 }
\ No newline at end of file
