Edigrating CL 12464908

SACN Universe not ticking when slate UI not in focus. Replace with core delegate.
It block the ticking in ThrottleManager

Using
FCoreDelegates::OnEndFrame.Broadcast();

Instead tickable object

diff --git a/Engine/Plugins/Experimental/DigitalMultiplex/DMXProtocol/Source/DMXProtocol/Public/Interfaces/IDMXProtocolUniverse.h b/Engine/Plugins/Experimental/DigitalMultiplex/DMXProtocol/Source/DMXProtocol/Public/Interfaces/IDMXProtocolUniverse.h
--- a/Engine/Plugins/Experimental/DigitalMultiplex/DMXProtocol/Source/DMXProtocol/Public/Interfaces/IDMXProtocolUniverse.h
+++ b/Engine/Plugins/Experimental/DigitalMultiplex/DMXProtocol/Source/DMXProtocol/Public/Interfaces/IDMXProtocolUniverse.h
@@ -17,4 +17,7 @@ public:
 	virtual uint32 GetUniverseID() const = 0;
 	virtual TSharedPtr<FJsonObject> GetSettings() const = 0;
 	virtual bool IsSupportRDM() const = 0;
+
+	/** Tick on the end of each frame */
+	virtual void Tick(float DeltaTime) {};
 };
diff --git a/Engine/Plugins/Experimental/DigitalMultiplex/DMXProtocol/Source/DMXProtocolSACN/Private/DMXProtocolSACN.cpp b/Engine/Plugins/Experimental/DigitalMultiplex/DMXProtocol/Source/DMXProtocolSACN/Private/DMXProtocolSACN.cpp
--- a/Engine/Plugins/Experimental/DigitalMultiplex/DMXProtocol/Source/DMXProtocolSACN/Private/DMXProtocolSACN.cpp
+++ b/Engine/Plugins/Experimental/DigitalMultiplex/DMXProtocol/Source/DMXProtocolSACN/Private/DMXProtocolSACN.cpp
@@ -3,6 +3,8 @@
 #include "DMXProtocolSACN.h"
 #include "DMXProtocolTransportSACN.h"
 #include "Common/UdpSocketBuilder.h"
+#include "Misc/CoreDelegates.h"
+#include "Misc/App.h"
 
 #include "Packets/DMXProtocolE131PDUPacket.h"
 #include "DMXProtocolPackager.h"
@@ -43,6 +45,9 @@ bool FDMXProtocolSACN::Init()
 		UE_LOG_DMXPROTOCOL(Error, TEXT("%s:%s"), NetworkErrorMessagePrefix, *ErrorMessage);
 	}
 
+	// Declare Engine ticker
+	OnEndFrameHandle = FCoreDelegates::OnEndFrame.AddRaw(this, &FDMXProtocolSACN::OnEndFrame);
+
 	return true;
 }
 
@@ -51,6 +56,8 @@ bool FDMXProtocolSACN::Shutdown()
 	ReleaseNetworkInterface();
 	IDMXProtocol::OnNetworkInterfaceChanged.Remove(NetworkInterfaceChangedDelegate.GetHandle());
 
+	FCoreDelegates::OnEndFrame.Remove(OnEndFrameHandle);
+
 	return false;
 }
 
@@ -109,9 +116,22 @@ TSharedPtr<IDMXProtocolSender> FDMXProtocolSACN::GetSenderInterface() const
 
 bool FDMXProtocolSACN::Tick(float DeltaTime)
 {
+	for (const TPair<uint32, FDMXProtocolUniverseSACNPtr>& UniversePair : UniverseManager->GetAllUniverses())
+	{
+		if (FDMXProtocolUniverseSACNPtr Universe = UniversePair.Value)
+		{
+			Universe->Tick(DeltaTime);
+		}
+	}
+
 	return true;
 }
 
+void FDMXProtocolSACN::OnEndFrame()
+{
+	Tick(FApp::GetDeltaTime());
+}
+
 bool FDMXProtocolSACN::SendDiscovery(const TArray<uint16>& Universes)
 {
 	// Init Packager
diff --git a/Engine/Plugins/Experimental/DigitalMultiplex/DMXProtocol/Source/DMXProtocolSACN/Private/DMXProtocolSACN.h b/Engine/Plugins/Experimental/DigitalMultiplex/DMXProtocol/Source/DMXProtocolSACN/Private/DMXProtocolSACN.h
--- a/Engine/Plugins/Experimental/DigitalMultiplex/DMXProtocol/Source/DMXProtocolSACN/Private/DMXProtocolSACN.h
+++ b/Engine/Plugins/Experimental/DigitalMultiplex/DMXProtocol/Source/DMXProtocolSACN/Private/DMXProtocolSACN.h
@@ -18,6 +18,8 @@ class FDMXProtocolUniverseSACN;
 template<class TUniverse>
 class FDMXProtocolUniverseManager;
 
+using FDMXProtocolUniverseSACNPtr = TSharedPtr<FDMXProtocolUniverseSACN, ESPMode::ThreadSafe>;
+
 class DMXPROTOCOLSACN_API FDMXProtocolSACN
 	: public IDMXProtocol
 	, public IDMXNetworkInterface
@@ -76,6 +78,9 @@ public:
 private:
 	EDMXSendResult SendDMXInternal(uint16 UniverseID, const TSharedPtr<FDMXBuffer>& DMXBuffer) const;
 
+	/** Called each tick in LaunchEngineLoop */
+	void OnEndFrame();
+
 private:
 	FName ProtocolName;
 
@@ -98,4 +103,7 @@ private:
 	FOnNetworkInterfaceChangedDelegate NetworkInterfaceChangedDelegate;
 
 	const TCHAR* NetworkErrorMessagePrefix;
+
+	/** Called at the end of a frame, Allow to tick universes */
+	FDelegateHandle OnEndFrameHandle;
 };
diff --git a/Engine/Plugins/Experimental/DigitalMultiplex/DMXProtocol/Source/DMXProtocolSACN/Private/DMXProtocolUniverseSACN.cpp b/Engine/Plugins/Experimental/DigitalMultiplex/DMXProtocol/Source/DMXProtocolSACN/Private/DMXProtocolUniverseSACN.cpp
--- a/Engine/Plugins/Experimental/DigitalMultiplex/DMXProtocol/Source/DMXProtocolSACN/Private/DMXProtocolUniverseSACN.cpp
+++ b/Engine/Plugins/Experimental/DigitalMultiplex/DMXProtocol/Source/DMXProtocolSACN/Private/DMXProtocolUniverseSACN.cpp
@@ -146,6 +146,11 @@ bool FDMXProtocolUniverseSACN::ReceiveDMXBuffer()
 
 void FDMXProtocolUniverseSACN::Tick(float DeltaTime)
 {
+	if (!bIsTicking)
+	{
+		return;
+	}
+
 	if (TimeWithoutInputBufferRequestStart > TimeWithoutInputBufferRequestEnd)
 	{
 		// Stop listening and destroy the socket
@@ -162,16 +167,6 @@ void FDMXProtocolUniverseSACN::Tick(float DeltaTime)
 	}
 }
 
-bool FDMXProtocolUniverseSACN::IsTickable() const
-{
-	return bIsTicking;
-}
-
-TStatId FDMXProtocolUniverseSACN::GetStatId() const
-{
-	RETURN_QUICK_DECLARE_CYCLE_STAT(FDMXProtocolUniverseSACN, STATGROUP_Tickables);
-}
-
 void FDMXProtocolUniverseSACN::OnNetworkInterfaceChanged(const FString& InInterfaceIPAddress)
 {
 	FString ErrorMessage;
diff --git a/Engine/Plugins/Experimental/DigitalMultiplex/DMXProtocol/Source/DMXProtocolSACN/Private/DMXProtocolUniverseSACN.h b/Engine/Plugins/Experimental/DigitalMultiplex/DMXProtocol/Source/DMXProtocolSACN/Private/DMXProtocolUniverseSACN.h
--- a/Engine/Plugins/Experimental/DigitalMultiplex/DMXProtocol/Source/DMXProtocolSACN/Private/DMXProtocolUniverseSACN.h
+++ b/Engine/Plugins/Experimental/DigitalMultiplex/DMXProtocol/Source/DMXProtocolSACN/Private/DMXProtocolUniverseSACN.h
@@ -22,7 +22,6 @@ class ISocketSubsystem;
 class DMXPROTOCOLSACN_API FDMXProtocolUniverseSACN
 	: public IDMXProtocolUniverse
 	, public IDMXNetworkInterface
-	, public FTickableGameObject
 {
 public:
 	FDMXProtocolUniverseSACN(IDMXProtocolPtr InDMXProtocol, const FJsonObject& InSettings);
@@ -37,15 +36,9 @@ public:
 	virtual uint32 GetUniverseID() const override;
 	virtual TSharedPtr<FJsonObject> GetSettings() const override;
 	virtual bool IsSupportRDM() const override;
-	//~ End IDMXProtocolDevice implementation
 
-	//~ Begin FTickableObjectBase implementation
 	virtual void Tick(float DeltaTime) override;
-	virtual bool IsTickable() const;
-	virtual TStatId GetStatId() const override;
-	virtual bool IsTickableInEditor() const override { return true; }
-	virtual bool IsTickableWhenPaused() const override { return true; }
-	//~ End FTickableObjectBase implementation
+	//~ End IDMXProtocolDevice implementation
 
 	//~ Begin IDMXNetworkInterface implementation
 	virtual void OnNetworkInterfaceChanged(const FString& InInterfaceIPAddress) override;
