Fix UE4Game compilation issue in DigitalMultiplex automation tests.

diff --git a/Engine/Plugins/Experimental/DigitalMultiplex/DMXProtocol/Source/DMXProtocolSACN/Private/Tests/DMXPrtotocolSACNFlowTest.cpp b/Engine/Plugins/Experimental/DigitalMultiplex/DMXProtocol/Source/DMXProtocolSACN/Private/Tests/DMXPrtotocolSACNFlowTest.cpp
--- a/Engine/Plugins/Experimental/DigitalMultiplex/DMXProtocol/Source/DMXProtocolSACN/Private/Tests/DMXPrtotocolSACNFlowTest.cpp
+++ b/Engine/Plugins/Experimental/DigitalMultiplex/DMXProtocol/Source/DMXProtocolSACN/Private/Tests/DMXPrtotocolSACNFlowTest.cpp
@@ -11,6 +11,8 @@
 #include "DMXProtocolSACN.h"
 #include "DMXProtocolUniverseSACN.h"
 
+#if WITH_DEV_AUTOMATION_TESTS
+
 struct DMXPrtotocolSACNHelper
 {
 	DMXPrtotocolSACNHelper(uint16 InUniverseID, FAutomationTestBase* InTest)
@@ -197,4 +199,6 @@ bool FDMXPrtotocolSACNConsoleCommandsTest::RunTest(const FString& Parameters)
 	ADD_LATENT_AUTOMATION_COMMAND(FWaitLatentCommand(0.2f));
 
 	return true;
-}
\ No newline at end of file
+}
+
+#endif // WITH_DEV_AUTOMATION_TESTS
