DMX: Fix crash from FDMXBuffer shared pointer getting destroyed when used among different threads.

diff --git a/Engine/Plugins/Experimental/DigitalMultiplex/DMXEngine/Source/DMXEditor/Private/Widgets/SDMXInputInfo.cpp b/Engine/Plugins/Experimental/DigitalMultiplex/DMXEngine/Source/DMXEditor/Private/Widgets/SDMXInputInfo.cpp
--- a/Engine/Plugins/Experimental/DigitalMultiplex/DMXEngine/Source/DMXEditor/Private/Widgets/SDMXInputInfo.cpp
+++ b/Engine/Plugins/Experimental/DigitalMultiplex/DMXEngine/Source/DMXEditor/Private/Widgets/SDMXInputInfo.cpp
@@ -94,7 +94,7 @@ void SDMXInputInfo::UpdateChannelsValues()
 		{
 			if (IDMXProtocolUniversePtr Universe = DMXProtocolPtr->GetUniverseById(InfoSelecterPtr->GetCurrentUniverseID()))
 			{
-				if (TSharedPtr<FDMXBuffer> DMXBuffer = Universe->GetInputDMXBuffer())
+				if (FDMXBufferPtr DMXBuffer = Universe->GetInputDMXBuffer())
 				{
 					// check the sequence ID 
 					uint32 BufferSequenceID = DMXBuffer->GetSequenceID();
diff --git a/Engine/Plugins/Experimental/DigitalMultiplex/DMXEngine/Source/DMXRuntime/Private/DMXSubsystem.cpp b/Engine/Plugins/Experimental/DigitalMultiplex/DMXEngine/Source/DMXRuntime/Private/DMXSubsystem.cpp
--- a/Engine/Plugins/Experimental/DigitalMultiplex/DMXEngine/Source/DMXRuntime/Private/DMXSubsystem.cpp
+++ b/Engine/Plugins/Experimental/DigitalMultiplex/DMXEngine/Source/DMXRuntime/Private/DMXSubsystem.cpp
@@ -205,7 +205,7 @@ void UDMXSubsystem::GetRawBuffer(FDMXProtocolName SelectedProtocol, int32 Univer
 			TSharedPtr<IDMXProtocolUniverse, ESPMode::ThreadSafe> ProtocolUniverse = Protocol->GetUniverseById(UniverseIndex);
 			if (ProtocolUniverse.IsValid())
 			{
-				TSharedPtr<FDMXBuffer> Buffer = ProtocolUniverse.Get()->GetInputDMXBuffer();
+				FDMXBufferPtr Buffer = ProtocolUniverse.Get()->GetInputDMXBuffer();
 				if (Buffer.IsValid())
 				{
 					Buffer->AccessDMXData([&DMXBuffer](TArray<uint8>& InData)
@@ -322,7 +322,7 @@ bool UDMXSubsystem::GetFunctionsMap(UDMXEntityFixturePatch* InFixturePatch, cons
 		return false;
 	}
 
-	TSharedPtr<FDMXBuffer> InputDMXBuffer = Universe->GetInputDMXBuffer();
+	FDMXBufferPtr InputDMXBuffer = Universe->GetInputDMXBuffer();
 	if (!InputDMXBuffer.IsValid())
 	{
 		UE_LOG(DMXSubsystemLog, Warning, TEXT("InputDMXBuffer Not Valid"));
diff --git a/Engine/Plugins/Experimental/DigitalMultiplex/DMXProtocol/Source/DMXProtocol/Public/DMXProtocolCommon.h b/Engine/Plugins/Experimental/DigitalMultiplex/DMXProtocol/Source/DMXProtocol/Public/DMXProtocolCommon.h
--- a/Engine/Plugins/Experimental/DigitalMultiplex/DMXProtocol/Source/DMXProtocol/Public/DMXProtocolCommon.h
+++ b/Engine/Plugins/Experimental/DigitalMultiplex/DMXProtocol/Source/DMXProtocol/Public/DMXProtocolCommon.h
@@ -31,5 +31,6 @@ using FDMXPacketPtr = TSharedPtr<FDMXPacket, ESPMode::ThreadSafe>;
 using IDMXProtocolUniversePtr = TSharedPtr<IDMXProtocolUniverse, ESPMode::ThreadSafe>;
 using IDMXProtocolPtr = TSharedPtr<IDMXProtocol, ESPMode::ThreadSafe>;
 using IDMXProtocolPtrWeak = TWeakPtr<IDMXProtocol, ESPMode::ThreadSafe>;
+using FDMXBufferPtr = TSharedPtr<FDMXBuffer, ESPMode::ThreadSafe>;
 
 enum class EDMXSendResult : uint8;
diff --git a/Engine/Plugins/Experimental/DigitalMultiplex/DMXProtocol/Source/DMXProtocol/Public/Interfaces/IDMXProtocolUniverse.h b/Engine/Plugins/Experimental/DigitalMultiplex/DMXProtocol/Source/DMXProtocol/Public/Interfaces/IDMXProtocolUniverse.h
--- a/Engine/Plugins/Experimental/DigitalMultiplex/DMXProtocol/Source/DMXProtocol/Public/Interfaces/IDMXProtocolUniverse.h
+++ b/Engine/Plugins/Experimental/DigitalMultiplex/DMXProtocol/Source/DMXProtocol/Public/Interfaces/IDMXProtocolUniverse.h
@@ -10,8 +10,8 @@ public:
 	virtual ~IDMXProtocolUniverse() {}
 
 	virtual IDMXProtocolPtr GetProtocol() const = 0;
-	virtual TSharedPtr<FDMXBuffer> GetInputDMXBuffer() const = 0;
-	virtual TSharedPtr<FDMXBuffer> GetOutputDMXBuffer() const = 0;
+	virtual FDMXBufferPtr GetInputDMXBuffer() const = 0;
+	virtual FDMXBufferPtr GetOutputDMXBuffer() const = 0;
 	virtual bool SetDMXFragment(const IDMXFragmentMap& DMXFragment) = 0;
 	virtual uint8 GetPriority() const = 0;
 	virtual uint32 GetUniverseID() const = 0;
diff --git a/Engine/Plugins/Experimental/DigitalMultiplex/DMXProtocol/Source/DMXProtocolArtNet/Private/DMXProtocolArtNet.cpp b/Engine/Plugins/Experimental/DigitalMultiplex/DMXProtocol/Source/DMXProtocolArtNet/Private/DMXProtocolArtNet.cpp
--- a/Engine/Plugins/Experimental/DigitalMultiplex/DMXProtocol/Source/DMXProtocolArtNet/Private/DMXProtocolArtNet.cpp
+++ b/Engine/Plugins/Experimental/DigitalMultiplex/DMXProtocol/Source/DMXProtocolArtNet/Private/DMXProtocolArtNet.cpp
@@ -303,7 +303,7 @@ bool FDMXProtocolArtNet::Tick(float DeltaTime)
 	return true;
 }
 
-EDMXSendResult FDMXProtocolArtNet::SendDMXInternal(uint16 UniverseID, uint8 PortID, const TSharedPtr<FDMXBuffer>& DMXBuffer) const
+EDMXSendResult FDMXProtocolArtNet::SendDMXInternal(uint16 UniverseID, uint8 PortID, const FDMXBufferPtr& DMXBuffer) const
 {
 	// Init Packager
 	FDMXProtocolPackager Packager;
diff --git a/Engine/Plugins/Experimental/DigitalMultiplex/DMXProtocol/Source/DMXProtocolArtNet/Private/DMXProtocolArtNet.h b/Engine/Plugins/Experimental/DigitalMultiplex/DMXProtocol/Source/DMXProtocolArtNet/Private/DMXProtocolArtNet.h
--- a/Engine/Plugins/Experimental/DigitalMultiplex/DMXProtocol/Source/DMXProtocolArtNet/Private/DMXProtocolArtNet.h
+++ b/Engine/Plugins/Experimental/DigitalMultiplex/DMXProtocol/Source/DMXProtocolArtNet/Private/DMXProtocolArtNet.h
@@ -97,7 +97,7 @@ public:
 	bool TransmitRDM(uint32 InUniverseID, const TArray<uint8>& Data);
 
 private:
-	EDMXSendResult SendDMXInternal(uint16 InUniverseID, uint8 InPort, const TSharedPtr<FDMXBuffer>& DMXBuffer) const;
+	EDMXSendResult SendDMXInternal(uint16 InUniverseID, uint8 InPort, const FDMXBufferPtr& DMXBuffer) const;
 
 	//~ Only the factory makes instances
 	FDMXProtocolArtNet() = delete;
diff --git a/Engine/Plugins/Experimental/DigitalMultiplex/DMXProtocol/Source/DMXProtocolArtNet/Private/DMXProtocolUniverseArtNet.cpp b/Engine/Plugins/Experimental/DigitalMultiplex/DMXProtocol/Source/DMXProtocolArtNet/Private/DMXProtocolUniverseArtNet.cpp
--- a/Engine/Plugins/Experimental/DigitalMultiplex/DMXProtocol/Source/DMXProtocolArtNet/Private/DMXProtocolUniverseArtNet.cpp
+++ b/Engine/Plugins/Experimental/DigitalMultiplex/DMXProtocol/Source/DMXProtocolArtNet/Private/DMXProtocolUniverseArtNet.cpp
@@ -18,8 +18,8 @@ FDMXProtocolUniverseArtNet::FDMXProtocolUniverseArtNet(IDMXProtocolPtr InDMXProt
 	checkf(WeakDMXProtocol.IsValid(), TEXT("DMXProtocol pointer is nullptr"));
 
 	Priority = 100;
-	OutputDMXBuffer = MakeShared<FDMXBuffer>();
-	InputDMXBuffer = MakeShared<FDMXBuffer>();
+	OutputDMXBuffer = MakeShared<FDMXBuffer, ESPMode::ThreadSafe>();
+	InputDMXBuffer = MakeShared<FDMXBuffer, ESPMode::ThreadSafe>();
 
 	Settings = MakeShared<FJsonObject>(InSettings);
 	checkf(Settings->HasField(TEXT("PortID")), TEXT("DMXProtocol PortID is not valid"));
@@ -45,12 +45,12 @@ IDMXProtocolPtr FDMXProtocolUniverseArtNet::GetProtocol() const
 	return WeakDMXProtocol.Pin();
 }
 
-TSharedPtr<FDMXBuffer> FDMXProtocolUniverseArtNet::GetOutputDMXBuffer() const
+FDMXBufferPtr FDMXProtocolUniverseArtNet::GetOutputDMXBuffer() const
 {
 	return OutputDMXBuffer;
 }
 
-TSharedPtr<FDMXBuffer> FDMXProtocolUniverseArtNet::GetInputDMXBuffer() const
+FDMXBufferPtr FDMXProtocolUniverseArtNet::GetInputDMXBuffer() const
 {
 	return InputDMXBuffer;
 }
diff --git a/Engine/Plugins/Experimental/DigitalMultiplex/DMXProtocol/Source/DMXProtocolArtNet/Private/DMXProtocolUniverseArtNet.h b/Engine/Plugins/Experimental/DigitalMultiplex/DMXProtocol/Source/DMXProtocolArtNet/Private/DMXProtocolUniverseArtNet.h
--- a/Engine/Plugins/Experimental/DigitalMultiplex/DMXProtocol/Source/DMXProtocolArtNet/Private/DMXProtocolUniverseArtNet.h
+++ b/Engine/Plugins/Experimental/DigitalMultiplex/DMXProtocol/Source/DMXProtocolArtNet/Private/DMXProtocolUniverseArtNet.h
@@ -15,8 +15,8 @@ public:
 
 	//~ Begin IDMXProtocolDevice implementation
 	virtual IDMXProtocolPtr GetProtocol() const override;
-	virtual TSharedPtr<FDMXBuffer> GetInputDMXBuffer() const override;
-	virtual TSharedPtr<FDMXBuffer> GetOutputDMXBuffer() const override;
+	virtual FDMXBufferPtr GetInputDMXBuffer() const override;
+	virtual FDMXBufferPtr GetOutputDMXBuffer() const override;
 	virtual bool SetDMXFragment(const IDMXFragmentMap& DMXFragment) override;
 	virtual uint8 GetPriority() const override;
 	virtual uint32 GetUniverseID() const override;
@@ -69,8 +69,8 @@ private:
 
 private:
 	IDMXProtocolPtrWeak WeakDMXProtocol;
-	TSharedPtr<FDMXBuffer> OutputDMXBuffer;
-	TSharedPtr<FDMXBuffer> InputDMXBuffer;
+	FDMXBufferPtr OutputDMXBuffer;
+	FDMXBufferPtr InputDMXBuffer;
 	uint8 Priority;
 	uint16 UniverseID;
 
diff --git a/Engine/Plugins/Experimental/DigitalMultiplex/DMXProtocol/Source/DMXProtocolArtNet/Private/Tests/DMXPrtotocolArtNetFlowTest.cpp b/Engine/Plugins/Experimental/DigitalMultiplex/DMXProtocol/Source/DMXProtocolArtNet/Private/Tests/DMXPrtotocolArtNetFlowTest.cpp
--- a/Engine/Plugins/Experimental/DigitalMultiplex/DMXProtocol/Source/DMXProtocolArtNet/Private/Tests/DMXPrtotocolArtNetFlowTest.cpp
+++ b/Engine/Plugins/Experimental/DigitalMultiplex/DMXProtocol/Source/DMXProtocolArtNet/Private/Tests/DMXPrtotocolArtNetFlowTest.cpp
@@ -42,7 +42,7 @@ bool FDMXPrtotocolArtNetBasicFlowTest::RunTest(const FString& Parameters)
 
 	// Wait before recieve DMX
 	FPlatformProcess::Sleep(0.2f);
-	TSharedPtr<FDMXBuffer> InputDMXBuffer = UniverseArtNet->GetInputDMXBuffer();
+	FDMXBufferPtr InputDMXBuffer = UniverseArtNet->GetInputDMXBuffer();
 
 	TestEqual(TEXT("Incoming buffer should be same"), FixtureValues[0], InputDMXBuffer->GetDMXDataAddress(FixtureChannels[0] - 1));
 	TestEqual(TEXT("Incoming buffer should be same"), FixtureValues[1], InputDMXBuffer->GetDMXDataAddress(FixtureChannels[1] - 1));
@@ -141,7 +141,7 @@ bool FDMXPrtotocolArtNetConsoleCommandsTest::RunTest(const FString& Parameters)
 
 	FPlatformProcess::Sleep(0.2f);
 
-	TSharedPtr<FDMXBuffer> InputDMXBuffer = Universe->GetInputDMXBuffer();
+	FDMXBufferPtr InputDMXBuffer = Universe->GetInputDMXBuffer();
 	TestEqual(TEXT("Incoming buffer should be same"), Values[0][1], InputDMXBuffer->GetDMXDataAddress(Values[0][0] - 1));
 	TestEqual(TEXT("Incoming buffer should be same"), Values[1][1], InputDMXBuffer->GetDMXDataAddress(Values[1][0] - 1));
 	TestEqual(TEXT("Incoming buffer should be same"), Values[2][1], InputDMXBuffer->GetDMXDataAddress(Values[2][0] - 1));
diff --git a/Engine/Plugins/Experimental/DigitalMultiplex/DMXProtocol/Source/DMXProtocolSACN/Private/DMXProtocolSACN.cpp b/Engine/Plugins/Experimental/DigitalMultiplex/DMXProtocol/Source/DMXProtocolSACN/Private/DMXProtocolSACN.cpp
--- a/Engine/Plugins/Experimental/DigitalMultiplex/DMXProtocol/Source/DMXProtocolSACN/Private/DMXProtocolSACN.cpp
+++ b/Engine/Plugins/Experimental/DigitalMultiplex/DMXProtocol/Source/DMXProtocolSACN/Private/DMXProtocolSACN.cpp
@@ -220,7 +220,7 @@ uint16 FDMXProtocolSACN::GetMaxUniverses() const
 	return ACN_MAX_UNIVERSES;
 }
 
-EDMXSendResult FDMXProtocolSACN::SendDMXInternal(uint16 InUniverseID, const TSharedPtr<FDMXBuffer>& DMXBuffer) const
+EDMXSendResult FDMXProtocolSACN::SendDMXInternal(uint16 InUniverseID, const FDMXBufferPtr& DMXBuffer) const
 {
 	// Init Packager
 	FDMXProtocolPackager Packager;
diff --git a/Engine/Plugins/Experimental/DigitalMultiplex/DMXProtocol/Source/DMXProtocolSACN/Private/DMXProtocolSACN.h b/Engine/Plugins/Experimental/DigitalMultiplex/DMXProtocol/Source/DMXProtocolSACN/Private/DMXProtocolSACN.h
--- a/Engine/Plugins/Experimental/DigitalMultiplex/DMXProtocol/Source/DMXProtocolSACN/Private/DMXProtocolSACN.h
+++ b/Engine/Plugins/Experimental/DigitalMultiplex/DMXProtocol/Source/DMXProtocolSACN/Private/DMXProtocolSACN.h
@@ -76,7 +76,7 @@ public:
 	static TSharedPtr<FInternetAddr> GetUniverseAddr(uint16 InUniverseID);
 
 private:
-	EDMXSendResult SendDMXInternal(uint16 UniverseID, const TSharedPtr<FDMXBuffer>& DMXBuffer) const;
+	EDMXSendResult SendDMXInternal(uint16 UniverseID, const FDMXBufferPtr& DMXBuffer) const;
 
 	/** Called each tick in LaunchEngineLoop */
 	void OnEndFrame();
diff --git a/Engine/Plugins/Experimental/DigitalMultiplex/DMXProtocol/Source/DMXProtocolSACN/Private/DMXProtocolUniverseSACN.cpp b/Engine/Plugins/Experimental/DigitalMultiplex/DMXProtocol/Source/DMXProtocolSACN/Private/DMXProtocolUniverseSACN.cpp
--- a/Engine/Plugins/Experimental/DigitalMultiplex/DMXProtocol/Source/DMXProtocolSACN/Private/DMXProtocolUniverseSACN.cpp
+++ b/Engine/Plugins/Experimental/DigitalMultiplex/DMXProtocol/Source/DMXProtocolSACN/Private/DMXProtocolUniverseSACN.cpp
@@ -38,8 +38,8 @@ FDMXProtocolUniverseSACN::FDMXProtocolUniverseSACN(IDMXProtocolPtr InDMXProtocol
 	UniverseID = Settings->GetNumberField(TEXT("UniverseID"));
 
 	// Allocate new buffers
-	OutputDMXBuffer = MakeShared<FDMXBuffer>();
-	InputDMXBuffer = MakeShared<FDMXBuffer>();
+	OutputDMXBuffer = MakeShared<FDMXBuffer, ESPMode::ThreadSafe>();
+	InputDMXBuffer = MakeShared<FDMXBuffer, ESPMode::ThreadSafe>();
 
 	// Set default IP address
 	InterfaceIPAddress = GetDefault<UDMXProtocolSettings>()->InterfaceIPAddress;
@@ -66,12 +66,12 @@ IDMXProtocolPtr FDMXProtocolUniverseSACN::GetProtocol() const
 	return WeakDMXProtocol.Pin();
 }
 
-TSharedPtr<FDMXBuffer> FDMXProtocolUniverseSACN::GetOutputDMXBuffer() const
+FDMXBufferPtr FDMXProtocolUniverseSACN::GetOutputDMXBuffer() const
 {
 	return OutputDMXBuffer;
 }
 
-TSharedPtr<FDMXBuffer> FDMXProtocolUniverseSACN::GetInputDMXBuffer() const
+FDMXBufferPtr FDMXProtocolUniverseSACN::GetInputDMXBuffer() const
 {
 	TimeWithoutInputBufferRequestStart = FPlatformTime::Seconds();
 	TimeWithoutInputBufferRequestEnd = TimeWithoutInputBufferRequestStart + TimeWithoutInputBufferRequest;
diff --git a/Engine/Plugins/Experimental/DigitalMultiplex/DMXProtocol/Source/DMXProtocolSACN/Private/DMXProtocolUniverseSACN.h b/Engine/Plugins/Experimental/DigitalMultiplex/DMXProtocol/Source/DMXProtocolSACN/Private/DMXProtocolUniverseSACN.h
--- a/Engine/Plugins/Experimental/DigitalMultiplex/DMXProtocol/Source/DMXProtocolSACN/Private/DMXProtocolUniverseSACN.h
+++ b/Engine/Plugins/Experimental/DigitalMultiplex/DMXProtocol/Source/DMXProtocolSACN/Private/DMXProtocolUniverseSACN.h
@@ -29,8 +29,8 @@ public:
 
 	//~ Begin IDMXProtocolDevice implementation
 	virtual IDMXProtocolPtr GetProtocol() const override;
-	virtual TSharedPtr<FDMXBuffer> GetInputDMXBuffer() const override;
-	virtual TSharedPtr<FDMXBuffer> GetOutputDMXBuffer() const override;
+	virtual FDMXBufferPtr GetInputDMXBuffer() const override;
+	virtual FDMXBufferPtr GetOutputDMXBuffer() const override;
 	virtual bool SetDMXFragment(const IDMXFragmentMap& DMXFragment) override;
 	virtual uint8 GetPriority() const override;
 	virtual uint32 GetUniverseID() const override;
@@ -69,8 +69,8 @@ private:
 
 private:
 	IDMXProtocolPtrWeak WeakDMXProtocol;
-	TSharedPtr<FDMXBuffer> OutputDMXBuffer;
-	TSharedPtr<FDMXBuffer> InputDMXBuffer;
+	FDMXBufferPtr OutputDMXBuffer;
+	FDMXBufferPtr InputDMXBuffer;
 	uint8 Priority;
 	uint32 UniverseID;
 	bool bIsRDMSupport;
diff --git a/Engine/Plugins/Experimental/DigitalMultiplex/DMXProtocol/Source/DMXProtocolSACN/Private/Tests/DMXPrtotocolSACNFlowTest.cpp b/Engine/Plugins/Experimental/DigitalMultiplex/DMXProtocol/Source/DMXProtocolSACN/Private/Tests/DMXPrtotocolSACNFlowTest.cpp
--- a/Engine/Plugins/Experimental/DigitalMultiplex/DMXProtocol/Source/DMXProtocolSACN/Private/Tests/DMXPrtotocolSACNFlowTest.cpp
+++ b/Engine/Plugins/Experimental/DigitalMultiplex/DMXProtocol/Source/DMXProtocolSACN/Private/Tests/DMXPrtotocolSACNFlowTest.cpp
@@ -73,7 +73,7 @@ bool DMXPrtotocolSACNBasicFlow_SendDMX1_Part2::Update()
 DEFINE_LATENT_AUTOMATION_COMMAND_ONE_PARAMETER(DMXPrtotocolSACNBasicFlow_CheckDMX1_Part3, TSharedPtr<DMXPrtotocolSACNHelper>, Helper);
 bool DMXPrtotocolSACNBasicFlow_CheckDMX1_Part3::Update()
 {
-	TSharedPtr<FDMXBuffer> InputDMXBuffer = Helper->Universe->GetInputDMXBuffer();
+	FDMXBufferPtr InputDMXBuffer = Helper->Universe->GetInputDMXBuffer();
 
 	// Test package ACNPacketIdentifier
 	{
@@ -94,7 +94,7 @@ bool DMXPrtotocolSACNBasicFlow_CheckDMX1_Part3::Update()
 DEFINE_LATENT_AUTOMATION_COMMAND_ONE_PARAMETER(DMXPrtotocolSACNBasicFlow_SendDMX2_Part4, TSharedPtr<DMXPrtotocolSACNHelper>, Helper);
 bool DMXPrtotocolSACNBasicFlow_SendDMX2_Part4::Update()
 {
-	TSharedPtr<FDMXBuffer> InputDMXBuffer = Helper->Universe->GetInputDMXBuffer();
+	FDMXBufferPtr InputDMXBuffer = Helper->Universe->GetInputDMXBuffer();
 
 	IDMXFragmentMap DMXFragmentMap;
 
@@ -109,7 +109,7 @@ bool DMXPrtotocolSACNBasicFlow_SendDMX2_Part4::Update()
 DEFINE_LATENT_AUTOMATION_COMMAND_ONE_PARAMETER(DMXPrtotocolSACNBasicFlow_CheckDMX2_Part5, TSharedPtr<DMXPrtotocolSACNHelper>, Helper);
 bool DMXPrtotocolSACNBasicFlow_CheckDMX2_Part5::Update()
 {
-	TSharedPtr<FDMXBuffer> InputDMXBuffer = Helper->Universe->GetInputDMXBuffer();
+	FDMXBufferPtr InputDMXBuffer = Helper->Universe->GetInputDMXBuffer();
 
 	Helper->Test->TestEqual(TEXT("Old value DMX input should be the same"), Helper->FixtureValues[0], InputDMXBuffer->GetDMXDataAddress(Helper->FixtureChannels[0] - 1));
 	Helper->Test->TestEqual(TEXT("Incoming buffer should be same"), Helper->FixtureValues[4], InputDMXBuffer->GetDMXDataAddress(Helper->FixtureChannels[4] - 1));
@@ -178,7 +178,7 @@ bool FDMXPrtotocolSACNConsoleCommands_SendDMX1_Part1::Update()
 DEFINE_LATENT_AUTOMATION_COMMAND_ONE_PARAMETER(FDMXPrtotocolSACNConsoleCommands_CheckDMX1_Part2, TSharedPtr<DMXPrtotocolSACNHelper>, Helper);
 bool FDMXPrtotocolSACNConsoleCommands_CheckDMX1_Part2::Update()
 {
-	TSharedPtr<FDMXBuffer> InputDMXBuffer = Helper->Universe->GetInputDMXBuffer();
+	FDMXBufferPtr InputDMXBuffer = Helper->Universe->GetInputDMXBuffer();
 	Helper->Test->TestEqual(TEXT("Incoming buffer should be same"), Helper->FixtureValues[0], InputDMXBuffer->GetDMXDataAddress(Helper->FixtureChannels[0] - 1));
 	Helper->Test->TestEqual(TEXT("Incoming buffer should be same"), Helper->FixtureValues[1], InputDMXBuffer->GetDMXDataAddress(Helper->FixtureChannels[1] - 1));
 	Helper->Test->TestEqual(TEXT("Incoming buffer should be same"), Helper->FixtureValues[2], InputDMXBuffer->GetDMXDataAddress(Helper->FixtureChannels[2] - 1));
