Fix PortGuid is not initialized properly even though its struct probably has a custom default constructor

diff --git a/Engine/Plugins/VirtualProduction/DMX/DMXProtocol/Source/DMXProtocol/Private/IO/DMXInputPortConfig.cpp b/Engine/Plugins/VirtualProduction/DMX/DMXProtocol/Source/DMXProtocol/Private/IO/DMXInputPortConfig.cpp
--- a/Engine/Plugins/VirtualProduction/DMX/DMXProtocol/Source/DMXProtocol/Private/IO/DMXInputPortConfig.cpp
+++ b/Engine/Plugins/VirtualProduction/DMX/DMXProtocol/Source/DMXProtocol/Private/IO/DMXInputPortConfig.cpp
@@ -17,6 +17,7 @@ FDMXInputPortConfig::FDMXInputPortConfig()
 	, LocalUniverseStart(1)
 	, NumUniverses(10)
 	, ExternUniverseStart(1)
+	, PortGuid()
 {
 	// May be called before the protocol module is loaded, at this point we only expect already sanetized structs
 	if (FModuleManager::Get().IsModuleLoaded("DMXProtocol"))
diff --git a/Engine/Plugins/VirtualProduction/DMX/DMXProtocol/Source/DMXProtocol/Private/IO/DMXOutputPortConfig.cpp b/Engine/Plugins/VirtualProduction/DMX/DMXProtocol/Source/DMXProtocol/Private/IO/DMXOutputPortConfig.cpp
--- a/Engine/Plugins/VirtualProduction/DMX/DMXProtocol/Source/DMXProtocol/Private/IO/DMXOutputPortConfig.cpp
+++ b/Engine/Plugins/VirtualProduction/DMX/DMXProtocol/Source/DMXProtocol/Private/IO/DMXOutputPortConfig.cpp
@@ -20,6 +20,7 @@ FDMXOutputPortConfig::FDMXOutputPortConfig()
 	, NumUniverses(10)
 	, ExternUniverseStart(1)
 	, Priority(100)
+	, PortGuid()
 {
 	// May be called before the protocol module is loaded, at this point we only expect already sanetized structs
 	if (FModuleManager::Get().IsModuleLoaded("DMXProtocol"))
