Remove validation for guid as it is expected that zero constructing a guid can generate a new guid instead

diff --git a/Engine/Plugins/VirtualProduction/DMX/DMXProtocol/Source/DMXProtocol/Public/IO/DMXInputPortConfig.h b/Engine/Plugins/VirtualProduction/DMX/DMXProtocol/Source/DMXProtocol/Public/IO/DMXInputPortConfig.h
--- a/Engine/Plugins/VirtualProduction/DMX/DMXProtocol/Source/DMXProtocol/Public/IO/DMXInputPortConfig.h
+++ b/Engine/Plugins/VirtualProduction/DMX/DMXProtocol/Source/DMXProtocol/Public/IO/DMXInputPortConfig.h
@@ -82,6 +82,6 @@ protected:
 	 * Unique identifier, shared with the port instance.
 	 * Note: This needs be BlueprintReadWrite to be accessible to property type customization, but is hidden by customization.
 	 */
-	UPROPERTY(Config, BlueprintReadWrite, EditDefaultsOnly, Category = "Port Config Guid")
+	UPROPERTY(Config, BlueprintReadWrite, EditDefaultsOnly, Category = "Port Config Guid", meta = (IgnoreForMemberInitializationTest))
 	FGuid PortGuid;
 };
diff --git a/Engine/Plugins/VirtualProduction/DMX/DMXProtocol/Source/DMXProtocol/Public/IO/DMXOutputPortConfig.h b/Engine/Plugins/VirtualProduction/DMX/DMXProtocol/Source/DMXProtocol/Public/IO/DMXOutputPortConfig.h
--- a/Engine/Plugins/VirtualProduction/DMX/DMXProtocol/Source/DMXProtocol/Public/IO/DMXOutputPortConfig.h
+++ b/Engine/Plugins/VirtualProduction/DMX/DMXProtocol/Source/DMXProtocol/Public/IO/DMXOutputPortConfig.h
@@ -93,6 +93,6 @@ protected:
 	 * Unique identifier, shared with the port instance.
 	 * Note: This needs be BlueprintReadWrite to be accessible to property type customization, but is hidden by customization.
 	 */
-	UPROPERTY(Config, BlueprintReadWrite, EditDefaultsOnly, Category = "Port Config Guid")
+	UPROPERTY(Config, BlueprintReadWrite, EditDefaultsOnly, Category = "Port Config Guid", meta = (IgnoreForMemberInitializationTest))
 	FGuid PortGuid;
 };
