Fixed some more places incorrectly testing against NAME_SIZE (as it includes the null terminator)

#rnx

diff --git a/Engine/Plugins/VirtualProduction/DMX/DMXProtocol/Source/DMXProtocol/Private/DMXProtocolUtils.cpp b/Engine/Plugins/VirtualProduction/DMX/DMXProtocol/Source/DMXProtocol/Private/DMXProtocolUtils.cpp
--- a/Engine/Plugins/VirtualProduction/DMX/DMXProtocol/Source/DMXProtocol/Private/DMXProtocolUtils.cpp
+++ b/Engine/Plugins/VirtualProduction/DMX/DMXProtocol/Source/DMXProtocol/Private/DMXProtocolUtils.cpp
@@ -61,9 +61,9 @@ FString FDMXProtocolUtils::GenerateUniqueNameFromExisting(const TSet<FString>& I
 		int32 CountLength = Count > 0 ? (int32)FGenericPlatformMath::LogX(10.0f, Count) + 2 : 2;
 
 		// If the length of the final string will be too long, cut off the end so we can fit the number
-		if (CountLength + BaseName.Len() > NAME_SIZE)
+		if (CountLength + BaseName.Len() >= NAME_SIZE)
 		{
-			BaseName = BaseName.Left(NAME_SIZE - CountLength);
+			BaseName = BaseName.Left(NAME_SIZE - CountLength - 1);
 		}
 
 		FinalName = FString::Printf(TEXT("%s_%d"), *BaseName, Count);
diff --git a/Engine/Source/Editor/SubobjectDataInterface/Private/SubobjectDataSubsystem.cpp b/Engine/Source/Editor/SubobjectDataInterface/Private/SubobjectDataSubsystem.cpp
--- a/Engine/Source/Editor/SubobjectDataInterface/Private/SubobjectDataSubsystem.cpp
+++ b/Engine/Source/Editor/SubobjectDataInterface/Private/SubobjectDataSubsystem.cpp
@@ -1828,10 +1828,10 @@ bool USubobjectDataSubsystem::IsValidRename(const FSubobjectDataHandle& Handle,
 				OutErrorMessage = LOCTEXT("RenameFailed_EngineReservedName", "This name is reserved for engine use.");
 				return false;
 			}
-			else if (NewTextStr.Len() > NAME_SIZE)
+			else if (NewTextStr.Len() >= NAME_SIZE)
 			{
 				FFormatNamedArguments Arguments;
-				Arguments.Add(TEXT("CharCount"), NAME_SIZE);
+				Arguments.Add(TEXT("CharCount"), NAME_SIZE - 1);
 				OutErrorMessage = FText::Format(LOCTEXT("ComponentRenameFailed_TooLong", "Component name must be less than {CharCount} characters long."), Arguments);
 				return false;
 			}
diff --git a/Engine/Source/Editor/UnrealEd/Private/EditorScriptingHelpers.cpp b/Engine/Source/Editor/UnrealEd/Private/EditorScriptingHelpers.cpp
--- a/Engine/Source/Editor/UnrealEd/Private/EditorScriptingHelpers.cpp
+++ b/Engine/Source/Editor/UnrealEd/Private/EditorScriptingHelpers.cpp
@@ -227,9 +227,9 @@ bool EditorScriptingHelpers::IsAValidPathForCreateNewAsset(const FString& Object
 	}
 
 	// Make sure we are not creating an FName that is too large
-	if (ObjectPath.Len() > NAME_SIZE)
+	if (ObjectPath.Len() >= NAME_SIZE)
 	{
-		OutFailureReason = TEXT("This asset name is too long (") + FString::FromInt(ObjectPath.Len()) + TEXT(" characters), the maximum is ") + FString::FromInt(NAME_SIZE) + TEXT(". Please choose a shorter name.");
+		OutFailureReason = TEXT("This asset name is too long (") + FString::FromInt(ObjectPath.Len()) + TEXT(" characters), the maximum is ") + FString::FromInt(NAME_SIZE - 1) + TEXT(". Please choose a shorter name.");
 		return false;
 	}
 
