DMX - Fix port changes in various UI representations when Ports are changed in project settings

diff --git a/Engine/Plugins/VirtualProduction/DMX/DMXProtocol/Source/DMXProtocolEditor/Private/Widgets/SDMXPortSelector.cpp b/Engine/Plugins/VirtualProduction/DMX/DMXProtocol/Source/DMXProtocolEditor/Private/Widgets/SDMXPortSelector.cpp
--- a/Engine/Plugins/VirtualProduction/DMX/DMXProtocol/Source/DMXProtocolEditor/Private/Widgets/SDMXPortSelector.cpp
+++ b/Engine/Plugins/VirtualProduction/DMX/DMXProtocol/Source/DMXProtocolEditor/Private/Widgets/SDMXPortSelector.cpp
@@ -263,13 +263,16 @@ TArray<TSharedPtr<FDMXPortSelectorItem>> SDMXPortSelector::MakeComboBoxSource()
 
 void SDMXPortSelector::GenerateWidgetsFromPorts()
 {
-	const TSharedPtr<FDMXPortSelectorItem> SelectedItem = [this]()
+	const FGuid SelectedGuid = [this]()
 	{
 		if (PortNameComboBox.IsValid())
 		{
-			return PortNameComboBox->GetSelectedItem();
+			if (TSharedPtr<FDMXPortSelectorItem> Item = PortNameComboBox->GetSelectedItem())
+			{
+				return Item->GetGuid();
+			}
 		}
-		return TSharedPtr<FDMXPortSelectorItem>();
+		return FGuid();
 	}();
 
 	ComboBoxSource = MakeComboBoxSource(); 
@@ -292,9 +295,14 @@ void SDMXPortSelector::GenerateWidgetsFromPorts()
 		PortNameComboBox->RefreshOptions();
 
 		// Try to restore the selection. Will trigger OnPortItemSelectionChanged
-		if (ComboBoxSource.Contains(SelectedItem))
+		const TSharedPtr<FDMXPortSelectorItem>* SelectedPortItemPtr = ComboBoxSource.FindByPredicate([SelectedGuid](const TSharedPtr<FDMXPortSelectorItem>& Item)
+			{
+				return Item->GetGuid() == SelectedGuid;
+			});
+
+		if (SelectedPortItemPtr)
 		{
-			PortNameComboBox->SetSelectedItem(SelectedItem);
+			PortNameComboBox->SetSelectedItem(*SelectedPortItemPtr);
 		}
 		else if (TSharedPtr<FDMXPortSelectorItem> NewSelection = GetFirstPortInComboBoxSource())
 		{
