DMX - Fix sACN unicast isn't received in the Editor

diff --git a/Engine/Plugins/VirtualProduction/DMX/DMXProtocol/Source/DMXProtocolSACN/Private/DMXProtocolSACNReceiver.cpp b/Engine/Plugins/VirtualProduction/DMX/DMXProtocol/Source/DMXProtocolSACN/Private/DMXProtocolSACNReceiver.cpp
--- a/Engine/Plugins/VirtualProduction/DMX/DMXProtocol/Source/DMXProtocolSACN/Private/DMXProtocolSACNReceiver.cpp
+++ b/Engine/Plugins/VirtualProduction/DMX/DMXProtocol/Source/DMXProtocolSACN/Private/DMXProtocolSACNReceiver.cpp
@@ -254,6 +254,29 @@ void FDMXProtocolSACNReceiver::Update(const FTimespan& SocketWaitTime)
 				const uint16 UniverseID = *UniverseIDPtr;
 				DistributeReceivedData(UniverseID, Reader);
 			}
+			// fallback to unicast (but only if the address is not a multicast one!)
+			else if ((MulticastIp & 0xF0000000) != 0xE0000000)
+			{
+				// safely (and quickly) retrieve the Universe
+				if (Reader->Num() >= ACN_DMX_ROOT_PACKAGE_SIZE + ACN_DMX_PDU_FRAMING_PACKAGE_SIZE)
+				{
+					// the Universe is in the last two bytes
+					constexpr int32 UniverseOffset = -2;
+					Reader->Seek(ACN_DMX_ROOT_PACKAGE_SIZE + ACN_DMX_PDU_FRAMING_PACKAGE_SIZE + UniverseOffset);
+					Reader->SetByteSwapping(true);
+					uint16 UniverseID = 0;
+					*Reader << UniverseID;
+					Reader->SetByteSwapping(false);
+					Reader->Seek(0);
+
+					// validate Universe
+					const uint32 FakeMulticastIp = GetIpForUniverseID(UniverseID);
+					if (MulticastGroupAddrToUniverseIDMap.Contains(FakeMulticastIp))
+					{
+						DistributeReceivedData(UniverseID, Reader);
+					}
+				}
+			}
 		}
 	}
 }
