DMX - IDMXProtocol::GetProtocols returns an empty array unless GetProtocol was called previously

diff --git a/Engine/Plugins/VirtualProduction/DMX/DMXProtocol/Source/DMXProtocol/Private/DMXProtocolModule.cpp b/Engine/Plugins/VirtualProduction/DMX/DMXProtocol/Source/DMXProtocol/Private/DMXProtocolModule.cpp
--- a/Engine/Plugins/VirtualProduction/DMX/DMXProtocol/Source/DMXProtocol/Private/DMXProtocolModule.cpp
+++ b/Engine/Plugins/VirtualProduction/DMX/DMXProtocol/Source/DMXProtocol/Private/DMXProtocolModule.cpp
@@ -45,7 +45,23 @@ void FDMXProtocolModule::RegisterProtocol(const FName& FactoryName, IDMXProtocol
 	// Broadcast protocol registered when all protocols are registered
 	if (NumRegisteredProtocols == NumProtocols)
 	{
-		OnProtocolsRegisteredDelegate.Broadcast();
+		// Create instances of all protocol implementations
+		for (const TTuple<FName, IDMXProtocolFactory*>& ProtocolNameToFactoryPair : DMXProtocolFactories)
+		{
+			IDMXProtocolPtr NewProtocol = ProtocolNameToFactoryPair.Value->CreateProtocol(ProtocolNameToFactoryPair.Key);
+			if (NewProtocol.IsValid())
+			{
+				DMXProtocols.Add(ProtocolNameToFactoryPair.Key, NewProtocol);
+
+				UE_LOG_DMXPROTOCOL(Log, TEXT("Creating protocol instance for: %s"), *ProtocolNameToFactoryPair.Key.ToString());
+			}
+			else
+			{
+				UE_LOG_DMXPROTOCOL(Verbose, TEXT("Unable to create Protocol %s"), *ProtocolNameToFactoryPair.Key.ToString());
+			}
+		}
+
+		OnProtocolsRegistered.Broadcast();
 	}
 }
 
@@ -73,38 +89,9 @@ FDMXProtocolModule& FDMXProtocolModule::Get()
 
 IDMXProtocolPtr FDMXProtocolModule::GetProtocol(const FName InProtocolName)
 {
-	IDMXProtocolPtr* DMXProtocolPtr = nullptr;
-	if (!InProtocolName.IsNone())
-	{
-		DMXProtocolPtr = DMXProtocols.Find(InProtocolName);
-		if (DMXProtocolPtr == nullptr)
-		{
-			IDMXProtocolFactory** DMXProtocolFactory = DMXProtocolFactories.Find(InProtocolName);
-
-			if (DMXProtocolFactory != nullptr)
-			{
-				UE_LOG_DMXPROTOCOL(Log, TEXT("Creating protocol instance for: %s"), *InProtocolName.ToString());
-
-				IDMXProtocolPtr NewProtocol = (*DMXProtocolFactory)->CreateProtocol(InProtocolName);
-				if (NewProtocol.IsValid())
-				{
-					DMXProtocols.Add(InProtocolName, NewProtocol);
-					DMXProtocolPtr = DMXProtocols.Find(InProtocolName);
-				}
-				else
-				{
-					bool* bNotedPreviously = DMXProtocolFailureNotes.Find(InProtocolName);
-					if (!bNotedPreviously || !(*bNotedPreviously))
-					{
-						UE_LOG_DMXPROTOCOL(Verbose, TEXT("Unable to create Protocol %s"), *InProtocolName.ToString());
-						DMXProtocolFailureNotes.Add(InProtocolName, true);
-					}
-				}
-			}
-		}
-	}
+	IDMXProtocolPtr* ProtocolPtr = DMXProtocols.Find(InProtocolName);
 
-	return (DMXProtocolPtr == nullptr) ? nullptr : *DMXProtocolPtr;
+	return ProtocolPtr ? *ProtocolPtr : nullptr;
 }
 
 const TMap<FName, IDMXProtocolFactory*>& FDMXProtocolModule::GetProtocolFactories() const
@@ -119,13 +106,13 @@ const TMap<FName, IDMXProtocolPtr>& FDMXProtocolModule::GetProtocols() const
 
 void FDMXProtocolModule::StartupModule()
 {
-	// Deffer setup to after all protocols begin registered
-	OnProtocolsRegisteredDelegate.AddRaw(this, &FDMXProtocolModule::OnProtocolsRegistered);
+	// Deffer initialization to after all protocols begin registered
+	OnProtocolsRegistered.AddRaw(this, &FDMXProtocolModule::HandleProtocolsRegistered);
 }
 
 void FDMXProtocolModule::ShutdownModule()
 {
-	OnProtocolsRegisteredDelegate.RemoveAll(this);
+	OnProtocolsRegistered.RemoveAll(this);
 
 	FDMXPortManager::ShutdownManager();
 
@@ -142,7 +129,7 @@ void FDMXProtocolModule::ShutdownModule()
 #endif // WITH_EDITOR
 }
 
-void FDMXProtocolModule::OnProtocolsRegistered()
+void FDMXProtocolModule::HandleProtocolsRegistered()
 {
 #if WITH_EDITOR
 	ISettingsModule* SettingsModule = FModuleManager::GetModulePtr<ISettingsModule>("Settings");
diff --git a/Engine/Plugins/VirtualProduction/DMX/DMXProtocol/Source/DMXProtocol/Public/DMXProtocolModule.h b/Engine/Plugins/VirtualProduction/DMX/DMXProtocol/Source/DMXProtocol/Public/DMXProtocolModule.h
--- a/Engine/Plugins/VirtualProduction/DMX/DMXProtocol/Source/DMXProtocol/Public/DMXProtocolModule.h
+++ b/Engine/Plugins/VirtualProduction/DMX/DMXProtocol/Source/DMXProtocol/Public/DMXProtocolModule.h
@@ -22,7 +22,7 @@ public:
 	void UnregisterProtocol(const FName& FactoryName);
 
 	/** Delegate called when all protocols are registered */
-	FSimpleMulticastDelegate OnProtocolsRegisteredDelegate;
+	FSimpleMulticastDelegate OnProtocolsRegistered;
 
 	/** Needs to be set to the number of protocol implementations */
 	static const int32 NumProtocols;
@@ -56,7 +56,7 @@ public:
 
 private:
 	/** Called when all protocols are registered */
-	void OnProtocolsRegistered();
+	void HandleProtocolsRegistered();
 
 	void ShutdownDMXProtocol(const FName& ProtocolName);
 	void ShutdownAllDMXProtocols();
@@ -64,5 +64,4 @@ private:
 private:
 	TMap<FName, IDMXProtocolFactory*> DMXProtocolFactories;
 	TMap<FName, IDMXProtocolPtr> DMXProtocols;
-	TMap<FName, bool> DMXProtocolFailureNotes;
 };
diff --git a/Engine/Plugins/VirtualProduction/DMX/DMXProtocol/Source/DMXProtocolEditor/Private/DMXProtocolEditorModule.cpp b/Engine/Plugins/VirtualProduction/DMX/DMXProtocol/Source/DMXProtocolEditor/Private/DMXProtocolEditorModule.cpp
--- a/Engine/Plugins/VirtualProduction/DMX/DMXProtocol/Source/DMXProtocolEditor/Private/DMXProtocolEditorModule.cpp
+++ b/Engine/Plugins/VirtualProduction/DMX/DMXProtocol/Source/DMXProtocolEditor/Private/DMXProtocolEditorModule.cpp
@@ -22,13 +22,13 @@ IMPLEMENT_MODULE( FDMXProtocolEditorModule, DMXProtocolEditor );
 void FDMXProtocolEditorModule::StartupModule()
 {
 	FDMXProtocolModule& ProtocolModule = FModuleManager::GetModuleChecked<FDMXProtocolModule>("DMXProtocol");
-	ProtocolModule.OnProtocolsRegisteredDelegate.AddRaw(this, &FDMXProtocolEditorModule::OnProtocolsRegistered);
+	ProtocolModule.OnProtocolsRegistered.AddRaw(this, &FDMXProtocolEditorModule::OnProtocolsRegistered);
 }
 
 void FDMXProtocolEditorModule::ShutdownModule()
 {
 	FDMXProtocolModule& ProtocolModule = FModuleManager::GetModuleChecked<FDMXProtocolModule>("DMXProtocol");
-	ProtocolModule.OnProtocolsRegisteredDelegate.RemoveAll(this);
+	ProtocolModule.OnProtocolsRegistered.RemoveAll(this);
 
 	UnregisterDetailsCustomizations();
 }
