DMX - Change Set Send DMX Enabled and Set Receive DMX Enabled methods and corresponding menu option and project settings to create and destroy sockets, add corresponding commandline args

diff --git a/Engine/Plugins/VirtualProduction/DMX/DMXProtocol/Source/DMXProtocol/Private/DMXProtocolSettings.cpp b/Engine/Plugins/VirtualProduction/DMX/DMXProtocol/Source/DMXProtocol/Private/DMXProtocolSettings.cpp
--- a/Engine/Plugins/VirtualProduction/DMX/DMXProtocol/Source/DMXProtocol/Private/DMXProtocolSettings.cpp
+++ b/Engine/Plugins/VirtualProduction/DMX/DMXProtocol/Source/DMXProtocol/Private/DMXProtocolSettings.cpp
@@ -83,8 +83,18 @@ void UDMXProtocolSettings::PostInitProperties()
 		Attribute.CleanupKeywords();
 	}
 
-	bOverrideSendDMXEnabled = bDefaultSendDMXEnabled;
-	bOverrideReceiveDMXEnabled = bDefaultReceiveDMXEnabled;
+	// Parse command line options for send and receive dmx
+	if (FParse::Bool(FCommandLine::Get(), TEXT("DEFAULTSENDDMXENABLED="), bDefaultSendDMXEnabled))
+	{
+		UE_LOG(LogDMXProtocol, Log, TEXT("Overridden Default Send DMX Enabled from command line, set to %s."), bDefaultSendDMXEnabled ? TEXT("True") : TEXT("False"));
+	}
+	OverrideSendDMXEnabled(bDefaultSendDMXEnabled);
+
+	if (FParse::Bool(FCommandLine::Get(), TEXT("DEFAULTRECEIVEDMXENABLED="), bDefaultReceiveDMXEnabled))
+	{
+		UE_LOG(LogDMXProtocol, Log, TEXT("Overridden Default Receive DMX Enabled from command line, set to %s."), bDefaultReceiveDMXEnabled ? TEXT("True") : TEXT("False"));
+	}
+	OverrideReceiveDMXEnabled(bDefaultReceiveDMXEnabled);
 }
 
 #if WITH_EDITOR
diff --git a/Engine/Plugins/VirtualProduction/DMX/DMXProtocol/Source/DMXProtocol/Private/IO/DMXInputPort.cpp b/Engine/Plugins/VirtualProduction/DMX/DMXProtocol/Source/DMXProtocol/Private/IO/DMXInputPort.cpp
--- a/Engine/Plugins/VirtualProduction/DMX/DMXProtocol/Source/DMXProtocol/Private/IO/DMXInputPort.cpp
+++ b/Engine/Plugins/VirtualProduction/DMX/DMXProtocol/Source/DMXProtocol/Private/IO/DMXInputPort.cpp
@@ -69,7 +69,7 @@ void FDMXInputPort::UpdateFromConfig(const FDMXInputPortConfig& InputPortConfig)
 	// Find if the port needs update its registration with the protocol
 	const bool bNeedsUpdateRegistration = [this, &InputPortConfig]()
 	{
-		if (!IsRegistered())
+		if (IsRegistered() != bReceiveDMXEnabled)
 		{
 			return true;
 		}
@@ -108,7 +108,7 @@ void FDMXInputPort::UpdateFromConfig(const FDMXInputPortConfig& InputPortConfig)
 	// Re-register the port if required
 	if (bNeedsUpdateRegistration)
 	{
-		if (IsValidPortSlow())
+		if (IsValidPortSlow() && bReceiveDMXEnabled)
 		{
 			Register();
 		}
@@ -264,6 +264,8 @@ bool FDMXInputPort::GameThreadGetDMXSignalFromRemoteUniverse(FDMXSignalSharedPtr
 void FDMXInputPort::OnSetReceiveDMXEnabled(bool bEnabled)
 {
 	bReceiveDMXEnabled = bEnabled;
+
+	UpdateFromConfig(*FindInputPortConfigChecked());
 }
 
 const FDMXInputPortConfig* FDMXInputPort::FindInputPortConfigChecked() const
diff --git a/Engine/Plugins/VirtualProduction/DMX/DMXProtocol/Source/DMXProtocol/Private/IO/DMXOutputPort.cpp b/Engine/Plugins/VirtualProduction/DMX/DMXProtocol/Source/DMXProtocol/Private/IO/DMXOutputPort.cpp
--- a/Engine/Plugins/VirtualProduction/DMX/DMXProtocol/Source/DMXProtocol/Private/IO/DMXOutputPort.cpp
+++ b/Engine/Plugins/VirtualProduction/DMX/DMXProtocol/Source/DMXProtocol/Private/IO/DMXOutputPort.cpp
@@ -72,7 +72,7 @@ void FDMXOutputPort::UpdateFromConfig(const FDMXOutputPortConfig& OutputPortConf
 	// Find if the port needs update its registration with the protocol
 	const bool bNeedsUpdateRegistration = [this, &OutputPortConfig]()
 	{
-		if (!IsRegistered())
+		if (IsRegistered() != CommunicationDeterminator.NeedsSendDMX())
 		{
 			return true;
 		}
@@ -116,7 +116,7 @@ void FDMXOutputPort::UpdateFromConfig(const FDMXOutputPortConfig& OutputPortConf
 	// Re-register the port if required
 	if (bNeedsUpdateRegistration)
 	{
-		if (IsValidPortSlow())
+		if (IsValidPortSlow() && CommunicationDeterminator.NeedsSendDMX())
 		{
 			Register();
 		}
@@ -351,11 +351,33 @@ bool FDMXOutputPort::GameThreadGetDMXSignalFromRemoteUniverse(FDMXSignalSharedPt
 void FDMXOutputPort::OnSetSendDMXEnabled(bool bEnabled)
 {
 	CommunicationDeterminator.SetSendEnabled(bEnabled);
+
+	UpdateFromConfig(*FindOutputPortConfigChecked());
 }
 
 void FDMXOutputPort::OnSetReceiveDMXEnabled(bool bEnabled)
 {
 	CommunicationDeterminator.SetReceiveEnabled(bEnabled);
+
+	UpdateFromConfig(*FindOutputPortConfigChecked());
+}
+
+const FDMXOutputPortConfig* FDMXOutputPort::FindOutputPortConfigChecked() const
+{
+	const UDMXProtocolSettings* ProjectSettings = GetDefault<UDMXProtocolSettings>();
+
+	for (const FDMXOutputPortConfig& OutputPortConfig : ProjectSettings->OutputPortConfigs)
+	{
+		if (OutputPortConfig.GetPortGuid() == PortGuid)
+		{
+			return &OutputPortConfig;
+		}
+	}
+
+	// Check failed, no config found
+	checkNoEntry();
+
+	return nullptr;
 }
 
 #undef LOCTEXT_NAMESPACE
diff --git a/Engine/Plugins/VirtualProduction/DMX/DMXProtocol/Source/DMXProtocol/Public/IO/DMXInputPort.h b/Engine/Plugins/VirtualProduction/DMX/DMXProtocol/Source/DMXProtocol/Public/IO/DMXInputPort.h
--- a/Engine/Plugins/VirtualProduction/DMX/DMXProtocol/Source/DMXProtocol/Public/IO/DMXInputPort.h
+++ b/Engine/Plugins/VirtualProduction/DMX/DMXProtocol/Source/DMXProtocol/Public/IO/DMXInputPort.h
@@ -126,7 +126,7 @@ private:
 	bool bRegistered;
 
 private:
-	/** Returns the port config with specified guid. Useful to find a port config by its port, as they share their guid */
+	/** Returns the port config that corresponds to the guid of this port. */
 	const FDMXInputPortConfig* FindInputPortConfigChecked() const;
 
 	/** The unique identifier of this port, shared with the port config this was constructed from. Should not be changed after construction. */
diff --git a/Engine/Plugins/VirtualProduction/DMX/DMXProtocol/Source/DMXProtocol/Public/IO/DMXOutputPort.h b/Engine/Plugins/VirtualProduction/DMX/DMXProtocol/Source/DMXProtocol/Public/IO/DMXOutputPort.h
--- a/Engine/Plugins/VirtualProduction/DMX/DMXProtocol/Source/DMXProtocol/Public/IO/DMXOutputPort.h
+++ b/Engine/Plugins/VirtualProduction/DMX/DMXProtocol/Source/DMXProtocol/Public/IO/DMXOutputPort.h
@@ -171,6 +171,10 @@ private:
 	/** True if the port is registered with it its protocol */
 	bool bRegistered;
 
+private:
+	/** Returns the port config that corresponds to the guid of this port. */
+	const FDMXOutputPortConfig* FindOutputPortConfigChecked() const;
+
 	/** The unique identifier of this port, shared with the port config this was constructed from. Should not be changed after construction. */
 	FGuid PortGuid;
 };
