DMX - Fix property change event raised on the wrong property handle

diff --git a/Engine/Plugins/VirtualProduction/DMX/DMXProtocol/Source/DMXProtocolEditor/Private/DetailsCustomizations/DMXPortConfigCustomizationBase .cpp b/Engine/Plugins/VirtualProduction/DMX/DMXProtocol/Source/DMXProtocolEditor/Private/DetailsCustomizations/DMXPortConfigCustomizationBase .cpp
--- a/Engine/Plugins/VirtualProduction/DMX/DMXProtocol/Source/DMXProtocolEditor/Private/DetailsCustomizations/DMXPortConfigCustomizationBase .cpp	
+++ b/Engine/Plugins/VirtualProduction/DMX/DMXProtocol/Source/DMXProtocolEditor/Private/DetailsCustomizations/DMXPortConfigCustomizationBase .cpp	
@@ -300,9 +300,9 @@ void FDMXPortConfigCustomizationBase::OnCommunicationTypeSelected()
 
 	const FScopedTransaction Transaction(LOCTEXT("CommunicationTypeSelected", "DMX: Selected Communication Type"));
 	
-	DeviceAddressHandle->NotifyPreChange();
+	CommunicationTypeHandle->NotifyPreChange();
 	CommunicationTypeHandle->SetValue(static_cast<uint8>(SelectedCommunicationType));
-	DeviceAddressHandle->NotifyPostChange(EPropertyChangeType::ValueSet);
+	CommunicationTypeHandle->NotifyPostChange(EPropertyChangeType::ValueSet);
 }
 
 void FDMXPortConfigCustomizationBase::OnIPAddressSelected()
