DMX - Fixes functionality of DMX replication for nDisplay (DMXDisplayCluster plugin), additionally

* No command line args are required by default for replication
* DMXDisplayCluster plugin cannot be overran by incoming data
* Fixes an issue where ports would never be destroyed
* Fixes an issue where adding new port configs in project settings causing socket errors to show while ArrayAdd operations are ongoing
* Removes C++ methods where new ports could be created in port manager but cannot be accessed from blueprints

diff --git a/Engine/Plugins/VirtualProduction/DMX/DMXDisplayCluster/Source/DMXDisplayCluster/Private/DMXDisplayClusterModule.cpp b/Engine/Plugins/VirtualProduction/DMX/DMXDisplayCluster/Source/DMXDisplayCluster/Private/DMXDisplayClusterModule.cpp
--- a/Engine/Plugins/VirtualProduction/DMX/DMXDisplayCluster/Source/DMXDisplayCluster/Private/DMXDisplayClusterModule.cpp
+++ b/Engine/Plugins/VirtualProduction/DMX/DMXDisplayCluster/Source/DMXDisplayCluster/Private/DMXDisplayClusterModule.cpp
@@ -1,97 +1,29 @@
 // Copyright Epic Games, Inc. All Rights Reserved.
 
 #include "DMXDisplayClusterModule.h"
-#include "Serialization/ArrayWriter.h"
-#include "Serialization/MemoryReader.h"
 
-#define LOCTEXT_NAMESPACE "DMXDisplayClusterModule"
+#include "DMXDisplayClusterReplicator.h"
+#include "IDisplayCluster.h"
 
-namespace
-{
-	/**
-	 * Event ID for the Cluster event used to replicate DMX for nDisplay
-	 *
-	 * As per 4.26 there is no registry for system events, hence
-	 * 0xDDDD0000 is used for DMX ndisplay replication system event by agreement 4.26
-	 */
-	constexpr int32 DMXnDisplayReplicationEventID = 0xDDDD0000;
-}
 
 void FDMXDisplayClusterModule::StartupModule()
 {
-	bForceMaster = FParse::Param(FCommandLine::Get(), TEXT("dc_dmx_master"));
-	bForceSlave = FParse::Param(FCommandLine::Get(), TEXT("dc_dmx_slave"));
-	bAutoMode = FParse::Param(FCommandLine::Get(), TEXT("dc_dmx"));
-
-	IDisplayCluster::Get().OnDisplayClusterStartSession().AddRaw(this, &FDMXDisplayClusterModule::SetupDisplayCluster);
-}
-
-void FDMXDisplayClusterModule::SendToSlaves(FDMXInputPortSharedRef DMXInputPortRef, FDMXSignalSharedRef DMXSignalRef)
-{
-	if (IDisplayClusterClusterManager* ClusterManager = IDisplayCluster::Get().GetClusterMgr())
-	{
-		FArrayWriter ArrayWriter;
-		DMXSignalRef->Serialize(ArrayWriter);
-
-		FDisplayClusterClusterEventBinary ClusterEvent;
-		ClusterEvent.bIsSystemEvent = true;
-		ClusterEvent.bShouldDiscardOnRepeat = false;
-		ClusterEvent.EventId = DMXnDisplayReplicationEventID;
-		ClusterEvent.EventData = MoveTemp(ArrayWriter);
-		constexpr bool bEmitFromMasterOnly = false;
-		ClusterManager->EmitClusterEventBinary(ClusterEvent, bEmitFromMasterOnly);
-	}
+	IDisplayCluster::Get().OnDisplayClusterStartSession().AddRaw(this, &FDMXDisplayClusterModule::CreateDMXDisplayClusterReplicator);
 }
 
-void FDMXDisplayClusterModule::SetupDisplayCluster()
+void FDMXDisplayClusterModule::CreateDMXDisplayClusterReplicator()
 {
-	if (IDisplayClusterClusterManager* ClusterManager = IDisplayCluster::Get().GetClusterMgr())
+	// Only run when operating as cluster and not in editor
+	if (IDisplayCluster::Get().GetOperationMode() == EDisplayClusterOperationMode::Cluster)
 	{
-		const bool bIsSlave = bForceSlave || (bAutoMode && ClusterManager->IsSlave());
-		if (bIsSlave)
+		if (IDisplayClusterClusterManager* ClusterManager = IDisplayCluster::Get().GetClusterMgr())
 		{
-			BinaryListener = FOnClusterEventBinaryListener::CreateRaw(this, &FDMXDisplayClusterModule::OnClusterEventReceived);
-			ClusterManager->AddClusterEventBinaryListener(BinaryListener);
-			// end here as we do not want to accidentally have a node acting both as a master and a slave
-			return;
-		}
-
-		const bool bIsMaster = bForceMaster || (bAutoMode && ClusterManager->IsMaster());
-		if (bIsMaster)
-		{
-			FDMXPortManager::Get().OnPortInputDequeued.AddRaw(this, &FDMXDisplayClusterModule::SendToSlaves);
+			DMXDisplayClusterReplicator = MakeShared<FDMXDisplayClusterReplicator>();
 		}
 	}
 }
 
 void FDMXDisplayClusterModule::ShutdownModule()
-{
-	if (IDisplayClusterClusterManager* ClusterManager = IDisplayCluster::Get().GetClusterMgr())
-	{
-		if (BinaryListener.IsBound())
-		{
-			ClusterManager->RemoveClusterEventBinaryListener(BinaryListener);
-		}
-	}
-}
-
-void FDMXDisplayClusterModule::OnClusterEventReceived(const FDisplayClusterClusterEventBinary& Event)
-{
-	if (Event.EventId == DMXnDisplayReplicationEventID)
-	{
-		FMemoryReader Reader(Event.EventData);
-		FDMXSignalSharedRef DMXSignal = MakeShared<FDMXSignal, ESPMode::ThreadSafe>();
-		DMXSignal->Serialize(Reader);
-
-		// forward the signal to all of the InputPorts
-		for (const FDMXInputPortSharedRef& DMXInputPortRef : FDMXPortManager::Get().GetInputPorts())
-		{
-			DMXInputPortRef->SingleProducerInputDMXSignal(DMXSignal);
-		}
-	}
-}
+{}
 
 IMPLEMENT_MODULE(FDMXDisplayClusterModule, DMXDisplayCluster)
-
-
-#undef LOCTEXT_NAMESPACE
diff --git a/Engine/Plugins/VirtualProduction/DMX/DMXDisplayCluster/Source/DMXDisplayCluster/Private/DMXDisplayClusterModule.h b/Engine/Plugins/VirtualProduction/DMX/DMXDisplayCluster/Source/DMXDisplayCluster/Private/DMXDisplayClusterModule.h
--- a/Engine/Plugins/VirtualProduction/DMX/DMXDisplayCluster/Source/DMXDisplayCluster/Private/DMXDisplayClusterModule.h
+++ b/Engine/Plugins/VirtualProduction/DMX/DMXDisplayCluster/Source/DMXDisplayCluster/Private/DMXDisplayClusterModule.h
@@ -4,16 +4,12 @@
 
 #include "CoreMinimal.h"
 #include "Modules/ModuleInterface.h"
-#include "Modules/ModuleManager.h"
-#include "IDisplayCluster.h"
-#include "DMXProtocolTypes.h"
-#include "IO/DMXPortManager.h"
-#include "IO/DMXInputPort.h"
-#include "Cluster/IDisplayClusterClusterManager.h"
-#include "Cluster/DisplayClusterClusterEvent.h"
 
+class FDMXDisplayClusterReplicator;
 
-class DMXDISPLAYCLUSTER_API FDMXDisplayClusterModule : public IModuleInterface
+
+class DMXDISPLAYCLUSTER_API FDMXDisplayClusterModule 
+	: public IModuleInterface
 {
 
 public:
@@ -21,25 +17,9 @@ public:
 	virtual void ShutdownModule() override;
 
 protected:
+	/** Setup the DisplayCluster replicator */
+	void CreateDMXDisplayClusterReplicator();
 
-	/** Delegate for DisplayCluster DMX events */
-	void OnClusterEventReceived(const FDisplayClusterClusterEventBinary& Event);
-
-	/** Listen for DisplayCluster DMX events */
-	FOnClusterEventBinaryListener BinaryListener;
-
-	/** Setup delegates for DisplayCluster */
-	void SetupDisplayCluster();
-
-	/** Forward a DMX signal to all of the configured DMX Slaves */
-	void SendToSlaves(FDMXInputPortSharedRef DMXInputPortRef, FDMXSignalSharedRef DMXSignalRef);
-
-	/** Force the node to be a DMX Master */
-	bool bForceMaster;
-
-	/** Force the node to be a DMX Slave */
-	bool bForceSlave;
-
-	/** Automatically map Master/Slave configuration to the DisplayCluster one */
-	bool bAutoMode;
+	/** The replicator instance */
+	TSharedPtr<FDMXDisplayClusterReplicator> DMXDisplayClusterReplicator;
 };
diff --git a/Engine/Plugins/VirtualProduction/DMX/DMXDisplayCluster/Source/DMXDisplayCluster/Private/DMXDisplayClusterReplicator.cpp b/Engine/Plugins/VirtualProduction/DMX/DMXDisplayCluster/Source/DMXDisplayCluster/Private/DMXDisplayClusterReplicator.cpp
new file mode 100644
--- /dev/null
+++ b/Engine/Plugins/VirtualProduction/DMX/DMXDisplayCluster/Source/DMXDisplayCluster/Private/DMXDisplayClusterReplicator.cpp
@@ -0,0 +1,201 @@
+// Copyright Epic Games, Inc. All Rights Reserved.
+
+#include "DMXDisplayClusterReplicator.h"
+
+#include "IO/DMXPortManager.h"
+#include "IO/DMXInputPort.h"
+#include "IO/DMXRawListener.h"
+
+#include "IDisplayCluster.h"
+#include "Cluster/DisplayClusterClusterEvent.h"
+#include "Serialization/ArrayWriter.h"
+#include "Serialization/MemoryReader.h"
+
+
+namespace
+{
+	/**
+	 * Event ID for the Cluster event used to replicate DMX for nDisplay
+	 *
+	 * As per 4.26 there is no registry for system events, hence
+	 * 0xDDDD0000 is used for DMX ndisplay replication system event by agreement 4.26
+	 */
+	constexpr int32 DMXnDisplayReplicationEventID = 0xDDDD0000;
+
+	/** 
+	 * A packet to replicate a dmx signal 
+	 */
+	struct FDMXDisplayClusterPacket
+	{
+		FDMXDisplayClusterPacket()
+			: Signal(MakeShared<FDMXSignal, ESPMode::ThreadSafe>())
+		{}
+
+		FDMXDisplayClusterPacket(const FDMXSignalSharedPtr& InSignal, int32 InPortIndex)
+			: Signal(InSignal)
+			, PortIndex(InPortIndex)
+		{}
+
+		void Serialize(FArchive& Ar)
+		{
+			Signal->Serialize(Ar);
+
+			Ar << PortIndex;
+		}
+
+		FDMXSignalSharedPtr Signal;
+
+		/** Index of the port where the packet was received */
+		int32 PortIndex = INDEX_NONE;
+	};
+}
+
+
+FDMXDisplayClusterReplicator::FDMXDisplayClusterReplicator()
+	: bForceMaster(false)
+	, bForceSlave(false)
+	, bClusterEventEmitter(false)
+	, TickType(ETickableTickType::Never)
+{
+	// Only supported when running in cluster mode without editor.
+	check(IDisplayCluster::Get().GetOperationMode() == EDisplayClusterOperationMode::Cluster);
+
+	if (IDisplayClusterClusterManager* ClusterManager = IDisplayCluster::Get().GetClusterMgr())
+	{
+		bForceMaster = FParse::Param(FCommandLine::Get(), TEXT("dc_dmx_master"));
+		bForceSlave = FParse::Param(FCommandLine::Get(), TEXT("dc_dmx_slave"));
+
+		ensureMsgf(!(bForceMaster && bForceSlave), TEXT("Ambigous command line for the DMX display cluster plugin. An instance cannot be 'dc_dmx_master' and 'dc_dmx_slave' at the same time."));
+
+		// Bind to OnClusterEventReceived. The master needs to bind here too to get data at the same time as slaves. 
+		BinaryListener = FOnClusterEventBinaryListener::CreateRaw(this, &FDMXDisplayClusterReplicator::OnClusterEventReceived);
+		ClusterManager->AddClusterEventBinaryListener(BinaryListener);
+
+		// Cache the input ports. As the application runs in cluster mode without editor, the Input ports array is not expected to change.
+		CachedInputPorts = FDMXPortManager::Get().GetInputPorts();
+
+		bClusterEventEmitter = [ClusterManager, this]()
+		{
+			const bool bAutoMode = !bForceMaster && !bForceSlave;
+
+			if (bAutoMode)
+			{
+				return ClusterManager->IsMaster();
+			}
+
+			return bForceMaster;
+		}();
+
+		if (bClusterEventEmitter)
+		{
+			TickType = ETickableTickType::Always;
+
+			// User a custom listener to avoid issues with tick precedence between the port and this object's tick here
+			for (const FDMXInputPortSharedRef& InputPort : CachedInputPorts)
+			{
+				// Don't use the default queue for the port, instead use a custom raw listener
+				InputPort->SetUseDefaultQueue(false);
+
+				TSharedPtr<FDMXRawListener> RawListener = MakeShared<FDMXRawListener>(InputPort);
+				RawListener->Start();
+
+				// Note we make sure here, and rely on the raw listener array is in the same order as the cached input ports array
+				RawListeners.Add(RawListener);
+			}
+		}
+		else
+		{
+			TickType = ETickableTickType::Never;
+
+			// If this is not an emitter, suspend receiving from protocols to not get any data from the network.
+			// This also allows to inject data into the ports - This is the single producer now.
+			FDMXPortManager::Get().SuspendProtocols();
+		}
+	}
+}
+
+FDMXDisplayClusterReplicator::~FDMXDisplayClusterReplicator()
+{
+	// Release raw listeners
+	for (const TSharedPtr<FDMXRawListener>& RawListener : RawListeners)
+	{
+		RawListener->Stop();
+	}
+
+	// Release binary listeners
+	if (IDisplayClusterClusterManager* ClusterManager = IDisplayCluster::Get().GetClusterMgr())
+	{
+		if (BinaryListener.IsBound())
+		{
+			ClusterManager->RemoveClusterEventBinaryListener(BinaryListener);
+		}
+	}
+}
+
+void FDMXDisplayClusterReplicator::OnClusterEventReceived(const FDisplayClusterClusterEventBinary& Event)
+{
+	if (Event.EventId == DMXnDisplayReplicationEventID)
+	{
+		FMemoryReader Reader(Event.EventData);
+
+		FDMXDisplayClusterPacket DMXDisplayClusterPacket;
+		DMXDisplayClusterPacket.Serialize(Reader);
+
+		if (ensureAlwaysMsgf(CachedInputPorts.IsValidIndex(DMXDisplayClusterPacket.PortIndex), TEXT("FDMXDisplayClusterReplicator: Input Ports specified in project settings are not identical across the nDipslay cluster. DMX replication failed.")))
+		{
+			CachedInputPorts[DMXDisplayClusterPacket.PortIndex]->GameThreadInjectDMXSignal(DMXDisplayClusterPacket.Signal.ToSharedRef());
+		}
+	}
+}
+
+void FDMXDisplayClusterReplicator::Tick(float DeltaTime)
+{
+	check(bClusterEventEmitter);
+
+	if (IDisplayClusterClusterManager* ClusterManager = IDisplayCluster::Get().GetClusterMgr())
+	{
+		for (int32 PortIndex = 0; PortIndex < CachedInputPorts.Num(); PortIndex++)
+		{
+			checkf(RawListeners.IsValidIndex(PortIndex), TEXT("The implementation relies on raw listeners corresponding to the ports array"));
+
+			// Get all data from the raw listener, but only keep the latest signals of this tick.
+			TSharedPtr<FDMXRawListener>& RawListener = RawListeners[PortIndex];
+
+			ExternUniverseToSignalForReplicationMap.Reset();
+
+			FDMXSignalSharedPtr Signal;
+			int32 DummyLocalUniverseID; // Can be ignored, we need the extern universe
+			while (RawListener->DequeueSignal(this, Signal, DummyLocalUniverseID))
+			{
+				ExternUniverseToSignalForReplicationMap.FindOrAdd(Signal->ExternUniverseID) = Signal;
+			}
+
+			// Emmit latest signals as cluster event
+			for (const TTuple<int32, FDMXSignalSharedPtr>& ExternUniverseToSignalPair : ExternUniverseToSignalForReplicationMap)
+			{
+				FDMXDisplayClusterPacket DMXDisplayClusterPacket(ExternUniverseToSignalPair.Value, PortIndex);
+
+				FArrayWriter ArrayWriter;
+				DMXDisplayClusterPacket.Serialize(ArrayWriter);
+
+				FDisplayClusterClusterEventBinary ClusterEvent;
+				ClusterEvent.bIsSystemEvent = true;
+				ClusterEvent.bShouldDiscardOnRepeat = false;
+				ClusterEvent.EventId = DMXnDisplayReplicationEventID;
+				ClusterEvent.EventData = MoveTemp(ArrayWriter);
+				constexpr bool bEmitFromMasterOnly = false;
+				ClusterManager->EmitClusterEventBinary(ClusterEvent, bEmitFromMasterOnly);
+			}
+		}
+	}
+}
+
+ETickableTickType FDMXDisplayClusterReplicator::GetTickableTickType() const
+{
+	return TickType;
+}
+
+TStatId FDMXDisplayClusterReplicator::GetStatId() const
+{
+	RETURN_QUICK_DECLARE_CYCLE_STAT(FDMXDisplayClusterReplicator, STATGROUP_Tickables);
+}
diff --git a/Engine/Plugins/VirtualProduction/DMX/DMXDisplayCluster/Source/DMXDisplayCluster/Private/DMXDisplayClusterReplicator.h b/Engine/Plugins/VirtualProduction/DMX/DMXDisplayCluster/Source/DMXDisplayCluster/Private/DMXDisplayClusterReplicator.h
new file mode 100644
--- /dev/null
+++ b/Engine/Plugins/VirtualProduction/DMX/DMXDisplayCluster/Source/DMXDisplayCluster/Private/DMXDisplayClusterReplicator.h
@@ -0,0 +1,58 @@
+// Copyright Epic Games, Inc. All Rights Reserved.
+
+#pragma once
+
+#include "DMXProtocolTypes.h"
+
+#include "CoreMinimal.h"
+#include "Tickable.h"
+#include "Cluster/IDisplayClusterClusterManager.h"
+
+class FDMXRawListener;
+
+
+/** Replicates input ports on tick */
+class DMXDISPLAYCLUSTER_API FDMXDisplayClusterReplicator
+	: public FTickableGameObject
+	, public TSharedFromThis<FDMXDisplayClusterReplicator>
+{
+public:
+	FDMXDisplayClusterReplicator();
+
+	virtual ~FDMXDisplayClusterReplicator();
+
+private:
+	/** Called whena cluster event was received */
+	void OnClusterEventReceived(const FDisplayClusterClusterEventBinary& Event);
+
+	//~Begin FTickableGameObject interface
+	virtual void Tick(float DeltaTime) override;
+	virtual bool IsTickableInEditor() const override { return true; }
+	virtual ETickableTickType GetTickableTickType() const override;
+	virtual TStatId GetStatId() const override;
+	//~End FTickableGameObject interface
+
+	/** Force the node to be a DMX Master */
+	bool bForceMaster;
+
+	/** Force the node to be a DMX Slave */
+	bool bForceSlave;
+
+	/** True if it is a cluster event emitter, else it is a cluster event receiver */
+	bool bClusterEventEmitter;
+
+	/** The tick type used (always for emitter, never for listener) */
+	ETickableTickType TickType;
+
+	/** Listen for DisplayCluster DMX events */
+	FOnClusterEventBinaryListener BinaryListener;
+
+	/** Input porst currently in use */
+	TArray<FDMXInputPortSharedRef> CachedInputPorts;
+
+	/** Raw listeners for the input ports */
+	TArray<TSharedPtr<FDMXRawListener>> RawListeners;
+
+	/** Map of latest Singals that needs be replicated per extern Universe. Used to only replicate latest data on tick. */
+	TMap<int32, FDMXSignalSharedPtr> ExternUniverseToSignalForReplicationMap;
+};
diff --git a/Engine/Plugins/VirtualProduction/DMX/DMXEngine/Source/DMXRuntime/Private/Library/DMXLibrary.cpp b/Engine/Plugins/VirtualProduction/DMX/DMXEngine/Source/DMXRuntime/Private/Library/DMXLibrary.cpp
--- a/Engine/Plugins/VirtualProduction/DMX/DMXEngine/Source/DMXRuntime/Private/Library/DMXLibrary.cpp
+++ b/Engine/Plugins/VirtualProduction/DMX/DMXEngine/Source/DMXRuntime/Private/Library/DMXLibrary.cpp
@@ -389,8 +389,8 @@ void UDMXLibrary::UpdatePorts()
 
 		if (!bInputPortExists)
 		{
-			// Default to disabled
-			bool bEnabled = false;
+			// Default to enabled
+			bool bEnabled = true;
 			PortReferences.InputPortReferences.Insert(FDMXInputPortReference(InputPortGuid, bEnabled), IndexInputPortConfig);
 		}
 	}
@@ -405,8 +405,8 @@ void UDMXLibrary::UpdatePorts()
 
 		if (!bOutputPortExists)
 		{
-			// Default to disabled
-			bool bEnabled = false;
+			// Default to enabled
+			bool bEnabled = true;
 			PortReferences.OutputPortReferences.Insert(FDMXOutputPortReference(OutputPortGuid, bEnabled), IndexOutputPortConfig);
 		}
 	}
@@ -485,8 +485,6 @@ void UDMXLibrary::UpgradeFromControllersToPorts()
 					});
 
 				*ExistingPortRef = FDMXInputPortReference(PortGuid, true);
-
-				FDMXPortManager::Get().UpdateFromProtocolSettings();
 			}
 
 			static void EnableOutputPortRef(UDMXLibrary* ThisLibrary, const FGuid& PortGuid)
@@ -498,8 +496,6 @@ void UDMXLibrary::UpgradeFromControllersToPorts()
 					});
 
 				*ExistingPortRef = FDMXOutputPortReference(PortGuid, true);
-
-				FDMXPortManager::Get().UpdateFromProtocolSettings();
 			}
 		};
 
@@ -584,8 +580,6 @@ void UDMXLibrary::UpgradeFromControllersToPorts()
 					InputPortConfig.NumUniverses = FixedLocalUniverseEnd - FixedLocalUniverseStart + 1;
 					InputPortConfig.ExternUniverseStart = FixedExternUniverseStart;
 
-					FDMXInputPortSharedRef InputPort = FDMXPortManager::Get().GetOrCreateInputPortFromConfig(InputPortConfig);
-
 					ProtocolSettings->InputPortConfigs.Add(InputPortConfig);
 
 					Local::EnableInputPortRef(this, InputPortConfig.GetPortGuid());
@@ -641,8 +635,6 @@ void UDMXLibrary::UpgradeFromControllersToPorts()
 					OutputPortConfig.ExternUniverseStart = FixedExternUniverseStart;
 					OutputPortConfig.bLoopbackToEngine = VoidController->CommunicationMode != EDMXCommunicationType::Broadcast;
 
-					FDMXOutputPortSharedRef OutputPort = FDMXPortManager::Get().GetOrCreateOutputPortFromConfig(OutputPortConfig);
-
 					ProtocolSettings->OutputPortConfigs.Add(OutputPortConfig);
 
 					Local::EnableOutputPortRef(this, OutputPortConfig.GetPortGuid());
@@ -688,8 +680,6 @@ void UDMXLibrary::UpgradeFromControllersToPorts()
 						OutputPortConfig.ExternUniverseStart = FixedExternUniverseStart;
 						OutputPortConfig.bLoopbackToEngine = false;
 
-						FDMXOutputPortSharedRef OutputPort = FDMXPortManager::Get().GetOrCreateOutputPortFromConfig(OutputPortConfig);
-
 						ProtocolSettings->OutputPortConfigs.Add(OutputPortConfig);
 
 						Local::EnableOutputPortRef(this, OutputPortConfig.GetPortGuid());
@@ -698,6 +688,10 @@ void UDMXLibrary::UpgradeFromControllersToPorts()
 			}
 		}
 
+		// Apply the changes to the port manager
+		FDMXPortManager::Get().UpdateFromProtocolSettings();
+
+		// Apply the changes to this library
 		UpdatePorts();
 	}
 
diff --git a/Engine/Plugins/VirtualProduction/DMX/DMXProtocol/Source/DMXProtocol/Private/DMXProtocolSettings.cpp b/Engine/Plugins/VirtualProduction/DMX/DMXProtocol/Source/DMXProtocol/Private/DMXProtocolSettings.cpp
--- a/Engine/Plugins/VirtualProduction/DMX/DMXProtocol/Source/DMXProtocol/Private/DMXProtocolSettings.cpp
+++ b/Engine/Plugins/VirtualProduction/DMX/DMXProtocol/Source/DMXProtocol/Private/DMXProtocolSettings.cpp
@@ -155,8 +155,7 @@ void UDMXProtocolSettings::PostEditChangeChainProperty(FPropertyChangedChainEven
 	else if	(
 		PropertyName == GET_MEMBER_NAME_CHECKED(UDMXProtocolSettings, InputPortConfigs) ||
 		PropertyName == GET_MEMBER_NAME_CHECKED(UDMXProtocolSettings, OutputPortConfigs) ||
-		(InputPortConfigStruct == PropertyOwnerStruct || OutputPortConfigStruct || PropertyOwnerStruct)
-		)
+		(InputPortConfigStruct == PropertyOwnerStruct) || (OutputPortConfigStruct == PropertyOwnerStruct))
 	{
 		// if a new config was added, create a guid for that. We cannot do that in the ctor because 
 		// the engine expects clearly defined default values, FGuid::NewGuid doesn't work with that.
@@ -199,4 +198,3 @@ void UDMXProtocolSettings::OverrideReceiveDMXEnabled(bool bEnabled)
 
 	OnSetReceiveDMXEnabled.Broadcast(bEnabled);
 }
-
diff --git a/Engine/Plugins/VirtualProduction/DMX/DMXProtocol/Source/DMXProtocol/Private/IO/DMXInputPort.cpp b/Engine/Plugins/VirtualProduction/DMX/DMXProtocol/Source/DMXProtocol/Private/IO/DMXInputPort.cpp
--- a/Engine/Plugins/VirtualProduction/DMX/DMXProtocol/Source/DMXProtocol/Private/IO/DMXInputPort.cpp
+++ b/Engine/Plugins/VirtualProduction/DMX/DMXProtocol/Source/DMXProtocol/Private/IO/DMXInputPort.cpp
@@ -15,22 +15,6 @@ DECLARE_CYCLE_STAT(TEXT("Input Port Tick"), STAT_DMXInputPortTick, STATGROUP_DMX
 
 #define LOCTEXT_NAMESPACE "DMXInputPort"
 
-FDMXInputPortSharedRef FDMXInputPort::Create()
-{
-	FDMXInputPortSharedRef NewInputPort = MakeShared<FDMXInputPort, ESPMode::ThreadSafe>();
-
-	NewInputPort->PortGuid = FGuid::NewGuid();
-
-	UDMXProtocolSettings* Settings = GetMutableDefault<UDMXProtocolSettings>();
-	check(Settings);
-
-	NewInputPort->bReceiveDMXEnabled = Settings->IsReceiveDMXEnabled();
-
-	// Bind to receive dmx changes
-	Settings->OnSetReceiveDMXEnabled.AddThreadSafeSP(NewInputPort, &FDMXInputPort::OnSetReceiveDMXEnabled);
-
-	return NewInputPort;
-}
 
 FDMXInputPortSharedRef FDMXInputPort::CreateFromConfig(const FDMXInputPortConfig& InputPortConfig)
 {
@@ -51,6 +35,8 @@ FDMXInputPortSharedRef FDMXInputPort::CreateFromConfig(const FDMXInputPortConfig
 
 	NewInputPort->UpdateFromConfig(InputPortConfig);
 
+	UE_LOG(LogDMXProtocol, VeryVerbose, TEXT("Created input port %s"), *NewInputPort->PortName);
+
 	return NewInputPort;
 }
 
@@ -62,10 +48,25 @@ FDMXInputPort::~FDMXInputPort()
 
 	// Port needs be unregistered before destruction
 	check(!bRegistered);
+
+	UE_LOG(LogDMXProtocol, VeryVerbose, TEXT("Destroyed input port %s"), *PortName);
 }
 
 void FDMXInputPort::UpdateFromConfig(const FDMXInputPortConfig& InputPortConfig)
 {
+	// Skip updating while the config has not yet initialized its device address 
+	if (!InputPortConfig.IsDeviceAddressInitialized())
+	{
+		return;
+	}
+
+	// Can only use configs that are in project settings
+	const UDMXProtocolSettings* ProtocolSettings = GetDefault<UDMXProtocolSettings>();
+	const bool bConfigIsInProjectSettings = ProtocolSettings->InputPortConfigs.ContainsByPredicate([&InputPortConfig](const FDMXInputPortConfig& Other) {
+		return InputPortConfig.GetPortGuid() == Other.GetPortGuid();
+		});
+	checkf(bConfigIsInProjectSettings, TEXT("Can only use configs that are in project settings"));
+
 	// Find if the port needs update its registration with the protocol
 	const bool bNeedsUpdateRegistration = [this, &InputPortConfig]()
 	{
@@ -74,6 +75,11 @@ void FDMXInputPort::UpdateFromConfig(const FDMXInputPortConfig& InputPortConfig)
 			return true;
 		}
 
+		if (IsRegistered() && FDMXPortManager::Get().AreProtocolsSuspended())
+		{
+			return true;
+		}
+
 		FName ProtocolName = Protocol.IsValid() ? Protocol->GetProtocolName() : NAME_None;
 
 		if (ProtocolName == InputPortConfig.ProtocolName &&
@@ -89,15 +95,13 @@ void FDMXInputPort::UpdateFromConfig(const FDMXInputPortConfig& InputPortConfig)
 	// Unregister the port if required before the new protocol is set
 	if (bNeedsUpdateRegistration)
 	{
-		if (IsRegistered())
-		{
-			Unregister();
-		}
+		Unregister();
 	}
 
 	Protocol = IDMXProtocol::Get(InputPortConfig.ProtocolName);
 
 	// Copy properties from the config into the base class
+	PortGuid = InputPortConfig.GetPortGuid();
 	CommunicationType = InputPortConfig.CommunicationType;
 	ExternUniverseStart = InputPortConfig.ExternUniverseStart;
 	DeviceAddress = InputPortConfig.DeviceAddress;
@@ -108,10 +112,7 @@ void FDMXInputPort::UpdateFromConfig(const FDMXInputPortConfig& InputPortConfig)
 	// Re-register the port if required
 	if (bNeedsUpdateRegistration)
 	{
-		if (IsValidPortSlow() && bReceiveDMXEnabled)
-		{
-			Register();
-		}
+		Register();
 	}
 
 	OnPortUpdated.Broadcast();
@@ -145,7 +146,7 @@ void FDMXInputPort::RemoveRawInput(TSharedRef<FDMXRawListener> RawInput)
 
 bool FDMXInputPort::Register()
 {
-	if (Protocol.IsValid())
+	if (IsValidPortSlow() && bReceiveDMXEnabled && !FDMXPortManager::Get().AreProtocolsSuspended())
 	{
 		bRegistered = Protocol->RegisterInputPort(SharedThis(this));
 
@@ -173,11 +174,14 @@ void FDMXInputPort::Tick(float DeltaTime)
 
 	// Tick universe Inputs with latest signals on tick
 	FDMXSignalSharedPtr Signal;
-	while (TickedBuffer.Dequeue(Signal))
+
+	if (bUseDefaultInputQueue)
 	{
-		// No need to filter extern universe, we already did that when enqueing
-		UniverseToLatestSignalMap.FindOrAdd(Signal->ExternUniverseID) = Signal;
-		FDMXPortManager::Get().OnPortInputDequeued.Broadcast(FDMXInputPort::SharedThis(this), Signal.ToSharedRef());
+		while (DefaultInputQueue.Dequeue(Signal))
+		{
+			// No need to filter extern universe, we already did that when enqueing
+			ExternUniverseToLatestSignalMap.FindOrAdd(Signal->ExternUniverseID) = Signal;
+		}
 	}
 }
 
@@ -194,7 +198,7 @@ TStatId FDMXInputPort::GetStatId() const
 void FDMXInputPort::ClearBuffers()
 {
 #if UE_BUILD_DEBUG
-	// Needs be called from the game thread, to maintain thread safety with UniverseToLatestSignalMap
+	// Needs be called from the game thread, to maintain thread safety with ExternUniverseToLatestSignalMap
 	check(IsInGameThread());
 #endif // UE_BUILD_DEBUG
 
@@ -203,8 +207,8 @@ void FDMXInputPort::ClearBuffers()
 		RawInput->ClearBuffer();
 	}
 
-	TickedBuffer.Empty();
-	UniverseToLatestSignalMap.Reset();
+	DefaultInputQueue.Empty();
+	ExternUniverseToLatestSignalMap.Reset();
 }
 
 void FDMXInputPort::SingleProducerInputDMXSignal(const FDMXSignalSharedRef& DMXSignal)
@@ -214,14 +218,15 @@ void FDMXInputPort::SingleProducerInputDMXSignal(const FDMXSignalSharedRef& DMXS
 		int32 ExternUniverseID = DMXSignal->ExternUniverseID;
 		if (IsExternUniverseInPortRange(ExternUniverseID))
 		{
-			int32 LocalUniverseID = ConvertExternToLocalUniverseID(ExternUniverseID);
-
 			for (const TSharedRef<FDMXRawListener>& RawListener : RawListeners)
 			{
 				RawListener->EnqueueSignal(this, DMXSignal);
 			}
 
-			TickedBuffer.Enqueue(DMXSignal);
+			if (bUseDefaultInputQueue)
+			{
+				DefaultInputQueue.Enqueue(DMXSignal);
+			}
 		}
 	}
 }
@@ -234,7 +239,7 @@ bool FDMXInputPort::GameThreadGetDMXSignal(int32 LocalUniverseID, FDMXSignalShar
 
 	int32 ExternUniverseID = ConvertLocalToExternUniverseID(LocalUniverseID);
 
-	const FDMXSignalSharedPtr* SignalPtr = UniverseToLatestSignalMap.Find(ExternUniverseID);
+	const FDMXSignalSharedPtr* SignalPtr = ExternUniverseToLatestSignalMap.Find(ExternUniverseID);
 	if (SignalPtr)
 	{
 		OutDMXSignal = *SignalPtr;
@@ -244,6 +249,15 @@ bool FDMXInputPort::GameThreadGetDMXSignal(int32 LocalUniverseID, FDMXSignalShar
 	return false;
 }
 
+const TMap<int32, FDMXSignalSharedPtr>& FDMXInputPort::GameThreadGetAllDMXSignals() const
+{
+#if UE_BUILD_DEBUG
+	check(IsInGameThread());
+#endif // UE_BUILD_DEBUG
+
+	return ExternUniverseToLatestSignalMap;
+}
+
 bool FDMXInputPort::GameThreadGetDMXSignalFromRemoteUniverse(FDMXSignalSharedPtr& OutDMXSignal, int32 RemoteUniverseID, bool bEvenIfNotLoopbackToEngine)
 {
 	// DEPRECATED 4.27
@@ -251,7 +265,7 @@ bool FDMXInputPort::GameThreadGetDMXSignalFromRemoteUniverse(FDMXSignalSharedPtr
 	check(IsInGameThread());
 #endif // UE_BUILD_DEBUG
 
-	const FDMXSignalSharedPtr* SignalPtr = UniverseToLatestSignalMap.Find(RemoteUniverseID);
+	const FDMXSignalSharedPtr* SignalPtr = ExternUniverseToLatestSignalMap.Find(RemoteUniverseID);
 	if (SignalPtr)
 	{
 		OutDMXSignal = *SignalPtr;
@@ -261,6 +275,22 @@ bool FDMXInputPort::GameThreadGetDMXSignalFromRemoteUniverse(FDMXSignalSharedPtr
 	return false;
 }
 
+void FDMXInputPort::GameThreadInjectDMXSignal(const FDMXSignalSharedRef& DMXSignal)
+{
+	ensureMsgf(!bUseDefaultInputQueue || !bReceiveDMXEnabled || FDMXPortManager::Get().AreProtocolsSuspended(), TEXT("Potential conflicts between injected and received signals, please revise the implementation."));
+
+	int32 ExternUniverseID = DMXSignal->ExternUniverseID;
+	if (IsExternUniverseInPortRange(ExternUniverseID))
+	{
+		ExternUniverseToLatestSignalMap.FindOrAdd(ExternUniverseID, DMXSignal) = DMXSignal;
+	}
+}
+
+void FDMXInputPort::SetUseDefaultQueue(bool bUse)
+{
+	bUseDefaultInputQueue = bUse;
+}
+
 void FDMXInputPort::OnSetReceiveDMXEnabled(bool bEnabled)
 {
 	bReceiveDMXEnabled = bEnabled;
diff --git a/Engine/Plugins/VirtualProduction/DMX/DMXProtocol/Source/DMXProtocol/Private/IO/DMXInputPortConfig.cpp b/Engine/Plugins/VirtualProduction/DMX/DMXProtocol/Source/DMXProtocol/Private/IO/DMXInputPortConfig.cpp
--- a/Engine/Plugins/VirtualProduction/DMX/DMXProtocol/Source/DMXProtocol/Private/IO/DMXInputPortConfig.cpp
+++ b/Engine/Plugins/VirtualProduction/DMX/DMXProtocol/Source/DMXProtocol/Private/IO/DMXInputPortConfig.cpp
@@ -11,9 +11,15 @@
 #include "Misc/Guid.h"
 
 
+namespace
+{
+	/** Identifier for unitialized ip adresses (see FDMXInputPortConfig::IsIPAddressInitialized()) */
+	static const FString InputDeviceAddressInitializationPending = TEXT("Initialization Pending");
+}
+
 FDMXInputPortConfig::FDMXInputPortConfig(const FGuid& InPortGuid)
 	: CommunicationType(EDMXCommunicationType::InternalOnly)
-	, DeviceAddress()
+	, DeviceAddress(InputDeviceAddressInitializationPending)
 	, LocalUniverseStart(1)
 	, NumUniverses(10)
 	, ExternUniverseStart(1)
@@ -35,6 +41,11 @@ const FGuid& FDMXInputPortConfig::GetPortGuid() const
 	return PortGuid;
 }
 
+bool FDMXInputPortConfig::IsDeviceAddressInitialized() const
+{
+	return DeviceAddress != InputDeviceAddressInitializationPending;
+}
+
 void FDMXInputPortConfig::SanetizePortName()
 {
 	if (!PortName.IsEmpty())
diff --git a/Engine/Plugins/VirtualProduction/DMX/DMXProtocol/Source/DMXProtocol/Private/IO/DMXOutputPort.cpp b/Engine/Plugins/VirtualProduction/DMX/DMXProtocol/Source/DMXProtocol/Private/IO/DMXOutputPort.cpp
--- a/Engine/Plugins/VirtualProduction/DMX/DMXProtocol/Source/DMXProtocol/Private/IO/DMXOutputPort.cpp
+++ b/Engine/Plugins/VirtualProduction/DMX/DMXProtocol/Source/DMXProtocol/Private/IO/DMXOutputPort.cpp
@@ -6,31 +6,12 @@
 #include "Interfaces/IDMXProtocol.h"
 #include "Interfaces/IDMXSender.h"
 #include "IO/DMXOutputPortConfig.h"
+#include "IO/DMXPortManager.h"
 #include "IO/DMXRawListener.h"
 
 
 #define LOCTEXT_NAMESPACE "DMXOutputPort"
 
-FDMXOutputPortSharedRef FDMXOutputPort::Create()
-{
-	FDMXOutputPortSharedRef NewOutputPort = MakeShared<FDMXOutputPort, ESPMode::ThreadSafe>();
-
-	NewOutputPort->PortGuid = FGuid::NewGuid();
-
-	UDMXProtocolSettings* Settings = GetMutableDefault<UDMXProtocolSettings>();
-	check(Settings);
-
-	NewOutputPort->CommunicationDeterminator.SetSendEnabled(Settings->IsSendDMXEnabled());
-	NewOutputPort->CommunicationDeterminator.SetReceiveEnabled(Settings->IsReceiveDMXEnabled());
-
-	// Bind to send dmx changes
-	Settings->OnSetSendDMXEnabled.AddThreadSafeSP(NewOutputPort, &FDMXOutputPort::OnSetSendDMXEnabled);
-
-	// Bind to receive dmx changes
-	Settings->OnSetReceiveDMXEnabled.AddThreadSafeSP(NewOutputPort, &FDMXOutputPort::OnSetReceiveDMXEnabled);
-
-	return NewOutputPort;
-}
 
 FDMXOutputPortSharedRef FDMXOutputPort::CreateFromConfig(const FDMXOutputPortConfig& OutputPortConfig)
 {
@@ -55,6 +36,8 @@ FDMXOutputPortSharedRef FDMXOutputPort::CreateFromConfig(const FDMXOutputPortCon
 
 	NewOutputPort->UpdateFromConfig(OutputPortConfig);
 
+	UE_LOG(LogDMXProtocol, VeryVerbose, TEXT("Created output port %s"), *NewOutputPort->PortName);
+
 	return NewOutputPort;
 }
 
@@ -65,10 +48,25 @@ FDMXOutputPort::~FDMXOutputPort()
 	
 	// Port needs be unregistered before destruction
 	check(!DMXSender);
+
+	UE_LOG(LogDMXProtocol, VeryVerbose, TEXT("Destroyed output port %s"), *PortName);
 }
 
 void FDMXOutputPort::UpdateFromConfig(const FDMXOutputPortConfig& OutputPortConfig)
 {
+	// Skip updating while the config has not yet initialized its device address 
+	if (!OutputPortConfig.IsDeviceAddressInitialized())
+	{
+		return;
+	}
+
+	// Can only use configs that are in project settings
+	const UDMXProtocolSettings* ProtocolSettings = GetDefault<UDMXProtocolSettings>();
+	const bool bConfigIsInProjectSettings = ProtocolSettings->OutputPortConfigs.ContainsByPredicate([&OutputPortConfig](const FDMXOutputPortConfig& Other) {
+		return OutputPortConfig.GetPortGuid() == Other.GetPortGuid();
+	});
+	checkf(bConfigIsInProjectSettings, TEXT("Can only use configs that are in project settings"));
+
 	// Find if the port needs update its registration with the protocol
 	const bool bNeedsUpdateRegistration = [this, &OutputPortConfig]()
 	{
@@ -102,6 +100,10 @@ void FDMXOutputPort::UpdateFromConfig(const FDMXOutputPortConfig& OutputPortConf
 	Protocol = IDMXProtocol::Get(OutputPortConfig.ProtocolName);
 
 	// Copy properties from the config
+	const FGuid& ConfigPortGuid = OutputPortConfig.GetPortGuid();
+	check(PortGuid.IsValid());
+	PortGuid = ConfigPortGuid;
+
 	CommunicationType = OutputPortConfig.CommunicationType;
 	ExternUniverseStart = OutputPortConfig.ExternUniverseStart;
 	DeviceAddress = OutputPortConfig.DeviceAddress;
@@ -116,10 +118,7 @@ void FDMXOutputPort::UpdateFromConfig(const FDMXOutputPortConfig& OutputPortConf
 	// Re-register the port if required
 	if (bNeedsUpdateRegistration)
 	{
-		if (IsValidPortSlow() && CommunicationDeterminator.NeedsSendDMX())
-		{
-			Register();
-		}
+		Register();
 	}
 
 	OnPortUpdated.Broadcast();
@@ -263,7 +262,7 @@ void FDMXOutputPort::SendDMXToRemoteUniverse(const TMap<int32, uint8>& ChannelTo
 
 bool FDMXOutputPort::Register()
 {
-	if (Protocol.IsValid())
+	if (IsValidPortSlow() && CommunicationDeterminator.IsSendDMXEnabled() && !FDMXPortManager::Get().AreProtocolsSuspended())
 	{
 		DMXSender = Protocol->RegisterOutputPort(SharedThis(this));
 
diff --git a/Engine/Plugins/VirtualProduction/DMX/DMXProtocol/Source/DMXProtocol/Private/IO/DMXOutputPortConfig.cpp b/Engine/Plugins/VirtualProduction/DMX/DMXProtocol/Source/DMXProtocol/Private/IO/DMXOutputPortConfig.cpp
--- a/Engine/Plugins/VirtualProduction/DMX/DMXProtocol/Source/DMXProtocol/Private/IO/DMXOutputPortConfig.cpp
+++ b/Engine/Plugins/VirtualProduction/DMX/DMXProtocol/Source/DMXProtocol/Private/IO/DMXOutputPortConfig.cpp
@@ -11,9 +11,15 @@
 #include "Misc/Guid.h"
 
 
+namespace
+{
+	/** Identifier for unitialized ip adresses (see FDMXInputPortConfig::IsIPAddressInitialized()) */
+	static const FString OutputDeviceAddressInitializationPending = TEXT("Initialization Pending");
+}
+
 FDMXOutputPortConfig::FDMXOutputPortConfig(const FGuid& InPortGuid)
 	: CommunicationType(EDMXCommunicationType::InternalOnly)
-	, DeviceAddress()
+	, DeviceAddress(OutputDeviceAddressInitializationPending)
 	, DestinationAddress()
 	, bLoopbackToEngine(true)
 	, LocalUniverseStart(1)
@@ -38,6 +44,11 @@ const FGuid& FDMXOutputPortConfig::GetPortGuid() const
 	return PortGuid;
 }
 
+bool FDMXOutputPortConfig::IsDeviceAddressInitialized() const
+{
+	return DeviceAddress != OutputDeviceAddressInitializationPending;
+}
+
 void FDMXOutputPortConfig::SanetizePortName()
 {
 	if (!PortName.IsEmpty())
diff --git a/Engine/Plugins/VirtualProduction/DMX/DMXProtocol/Source/DMXProtocol/Private/IO/DMXPortManager.cpp b/Engine/Plugins/VirtualProduction/DMX/DMXProtocol/Source/DMXProtocol/Private/IO/DMXPortManager.cpp
--- a/Engine/Plugins/VirtualProduction/DMX/DMXProtocol/Source/DMXProtocol/Private/IO/DMXPortManager.cpp
+++ b/Engine/Plugins/VirtualProduction/DMX/DMXProtocol/Source/DMXProtocol/Private/IO/DMXPortManager.cpp
@@ -28,74 +28,20 @@
 	return *CurrentManager;
 }
 
-FDMXInputPortSharedRef FDMXPortManager::CreateInputPort()
-{
-	FDMXInputPortSharedRef NewInputPort = FDMXInputPort::Create();
-	InputPorts.Add(NewInputPort);
-
-	OnPortsChanged.Broadcast();
-
-	return NewInputPort;
-}
-
-FDMXInputPortSharedRef FDMXPortManager::GetOrCreateInputPortFromConfig(const FDMXInputPortConfig& InputPortConfig)
+FDMXInputPortSharedRef FDMXPortManager::GetInputPortFromConfigChecked(const FDMXInputPortConfig& InputPortConfig)
 {
 	// The config needs a valid guid
 	check(InputPortConfig.GetPortGuid().IsValid());
 
-	FDMXInputPortSharedPtr ExistingPort = FindInputPortByGuid(InputPortConfig.GetPortGuid());
-	if (ExistingPort.IsValid())
-	{
-		return ExistingPort.ToSharedRef();
-	}
-
-	FDMXInputPortSharedRef NewInputPort = FDMXInputPort::CreateFromConfig(InputPortConfig);
-	InputPorts.Add(NewInputPort);
-
-	OnPortsChanged.Broadcast();
-
-	return NewInputPort;
+	return FindInputPortByGuidChecked(InputPortConfig.GetPortGuid());
 }
 
-void FDMXPortManager::RemoveInputPortChecked(const FGuid& PortGuid)
-{
-	FDMXInputPortSharedRef InputPort = FindInputPortByGuidChecked(PortGuid);
-	InputPorts.Remove(InputPort);
-}
-
-FDMXOutputPortSharedRef FDMXPortManager::CreateOutputPort()
-{
-	FDMXOutputPortSharedRef NewOutputPort = FDMXOutputPort::Create();
-	OutputPorts.Add(NewOutputPort);
-
-	OnPortsChanged.Broadcast();
-
-	return NewOutputPort;
-}
-
-FDMXOutputPortSharedRef FDMXPortManager::GetOrCreateOutputPortFromConfig(const FDMXOutputPortConfig& OutputPortConfig)
+FDMXOutputPortSharedRef FDMXPortManager::GetOutputPortFromConfigChecked(const FDMXOutputPortConfig& OutputPortConfig)
 {
 	// The config needs a valid guid
 	check(OutputPortConfig.GetPortGuid().IsValid());
 
-	FDMXOutputPortSharedPtr ExistingPort = FindOutputPortByGuid(OutputPortConfig.GetPortGuid());
-	if (ExistingPort.IsValid())
-	{
-		return ExistingPort.ToSharedRef();
-	}
-
-	FDMXOutputPortSharedRef NewOutputPort = FDMXOutputPort::CreateFromConfig(OutputPortConfig);
-	OutputPorts.Add(NewOutputPort);
-
-	OnPortsChanged.Broadcast();
-
-	return NewOutputPort;
+	return FindOutputPortByGuidChecked(OutputPortConfig.GetPortGuid());
 }
 
-void FDMXPortManager::RemoveOutputPortChecked(const FGuid& PortGuid)
-{
-	FDMXOutputPortSharedRef OutputPort = FindOutputPortByGuidChecked(PortGuid);
-	OutputPorts.Remove(OutputPort);
-}
-
 FDMXPortSharedPtr FDMXPortManager::FindPortByGuid(const FGuid& PortGuid) const
@@ -233,6 +179,9 @@
 		if (!ExistingInputPortConfig &&
 			PortGuidsFromProtocolSettings.Contains(PortGuid))
 		{
+			// Unregister the port to free a potential reference from its protocol
+			InputPort->Unregister();
+
 			InputPorts.Remove(InputPort);
 			PortGuidsFromProtocolSettings.Remove(PortGuid);
 		}
@@ -249,6 +198,9 @@
 		if (!ExistingOutputPortConfig &&
 			PortGuidsFromProtocolSettings.Contains(PortGuid))
 		{
+			// Unregister the port to free a potential references from its protocol
+			OutputPort->Unregister();
+
 			OutputPorts.Remove(OutputPort);
 			PortGuidsFromProtocolSettings.Remove(PortGuid);
 		}
@@ -276,6 +228,72 @@
 	OnPortsChanged.Broadcast();
 }
 
+void FDMXPortManager::SuspendProtocols()
+{
+	for (const FDMXInputPortSharedRef& InputPort : InputPorts)
+	{
+		InputPort->Unregister();
+	}
+
+	for (const FDMXOutputPortSharedRef& OutputPort : OutputPorts)
+	{
+		OutputPort->Unregister();
+	}
+
+	// Set this last to ensure thread safety of this function and AreProtocolsSuspended()
+	bProtocolsSuspended = true;
+}
+
+FDMXInputPortSharedRef FDMXPortManager::GetOrCreateInputPortFromConfig(const FDMXInputPortConfig& InputPortConfig)
+{
+	// The config needs a valid guid
+	check(InputPortConfig.GetPortGuid().IsValid());
+
+	FDMXInputPortSharedPtr ExistingPort = FindInputPortByGuid(InputPortConfig.GetPortGuid());
+	if (ExistingPort.IsValid())
+	{
+		return ExistingPort.ToSharedRef();
+	}
+
+	FDMXInputPortSharedRef NewInputPort = FDMXInputPort::CreateFromConfig(InputPortConfig);
+	InputPorts.Add(NewInputPort);
+
+	OnPortsChanged.Broadcast();
+
+	return NewInputPort;
+}
+
+void FDMXPortManager::RemoveInputPortChecked(const FGuid& PortGuid)
+{
+	FDMXInputPortSharedRef InputPort = FindInputPortByGuidChecked(PortGuid);
+	InputPorts.Remove(InputPort);
+}
+
+FDMXOutputPortSharedRef FDMXPortManager::GetOrCreateOutputPortFromConfig(const FDMXOutputPortConfig& OutputPortConfig)
+{
+	// The config needs a valid guid
+	check(OutputPortConfig.GetPortGuid().IsValid());
+
+	FDMXOutputPortSharedPtr ExistingPort = FindOutputPortByGuid(OutputPortConfig.GetPortGuid());
+	if (ExistingPort.IsValid())
+	{
+		return ExistingPort.ToSharedRef();
+	}
+
+	FDMXOutputPortSharedRef NewOutputPort = FDMXOutputPort::CreateFromConfig(OutputPortConfig);
+	OutputPorts.Add(NewOutputPort);
+
+	OnPortsChanged.Broadcast();
+
+	return NewOutputPort;
+}
+
+void FDMXPortManager::RemoveOutputPortChecked(const FGuid& PortGuid)
+{
+	FDMXOutputPortSharedRef OutputPort = FindOutputPortByGuidChecked(PortGuid);
+	OutputPorts.Remove(OutputPort);
+}
+
 void FDMXPortManager::StartupManager()
 {
 	UE_LOG(LogDMXProtocol, Verbose, TEXT("Starting up DMXPortManager"));
diff --git a/Engine/Plugins/VirtualProduction/DMX/DMXProtocol/Source/DMXProtocol/Public/IO/DMXInputPort.h b/Engine/Plugins/VirtualProduction/DMX/DMXProtocol/Source/DMXProtocol/Public/IO/DMXInputPort.h
--- a/Engine/Plugins/VirtualProduction/DMX/DMXProtocol/Source/DMXProtocol/Public/IO/DMXInputPort.h
+++ b/Engine/Plugins/VirtualProduction/DMX/DMXProtocol/Source/DMXProtocol/Public/IO/DMXInputPort.h
@@ -15,7 +15,7 @@
 struct FDMXInputPortConfig;
 class FDMXPortManager;
 class FDMXRawListener;
-
+class IDMXInputPortTickedBufferOverride;
 
 /**  
  * Higher level abstraction of a DMX input hiding networking specific and protocol specific complexity from the game. 
@@ -23,23 +23,19 @@ class FDMXRawListener;
  * To input DMX into your objects, refer to DMXRawListener and DMXTickedUniverseListener.
  *
  * Can only be constructed via DMXPortManger, see FDMXPortManager::CreateInputPort and FDMXPortManager::CreateInputPortFromConfig
-
  */
 class DMXPROTOCOL_API FDMXInputPort
 	: public FDMXPort
 	, public FTickableGameObject
 {
-	// Friend DMXPortManager so no other object can create instances
+	// Friend DMXPortManager so it can create instances and unregister void instances
 	friend FDMXPortManager;
 
-	// Friend Raw Listener so it can add and remove themselves to the port
+	// Friend Raw Listener so it can add and remove itself to the port
 	friend FDMXRawListener;
 
 protected:
-	/** Creates an output port that is not tied to a specific config. Hidden on purpose, use FDMXPortManager to create instances */
-	static FDMXInputPortSharedRef Create();
-
-	/** Creates an output port tied to a specific config. Hidden on purpose, use FDMXPortManager to create instances */
+	/** Creates an output port tied to a specific config */
 	static FDMXInputPortSharedRef CreateFromConfig(const FDMXInputPortConfig& InputPortConfig);
 
 public:
@@ -96,28 +92,49 @@ public:
 	/** Gets the last signal received in specified local universe. Returns false if no signal was received. Game-Thread only */
 	bool GameThreadGetDMXSignal(int32 LocalUniverseID, FDMXSignalSharedPtr& OutDMXSignal);
 
+	/** Gets all the last signal received. Game-Thread only */
+	const TMap<int32, FDMXSignalSharedPtr>& GameThreadGetAllDMXSignals() const;
+
 	/**  DEPRECATED 4.27. Gets the DMX signal from an extern (remote) Universe ID. */
 	UE_DEPRECATED(4.27, "Use GameThreadGetDMXSignal instead. GameThreadGetDMXSignalFromRemoteUniverse only exists to support deprecated blueprint nodes.")
 	bool GameThreadGetDMXSignalFromRemoteUniverse(FDMXSignalSharedPtr& OutDMXSignal, int32 RemoteUniverseID, bool bEvenIfNotLoopbackToEngine = false);
 
+	/** Returns true if receive DMX is disabled */
 	FORCEINLINE bool IsReceiveDMXEnabled() const { return bReceiveDMXEnabled; }
 
+	/** 
+	 * Injects a dmx signal into the game thread. 
+	 * Note, only signals get from game thread are affected but not raw listeners. 
+	 * Useful for nDisplay replication.
+	 */
+	void GameThreadInjectDMXSignal(const FDMXSignalSharedRef& DMXSignal);
+
+	/** 
+	 * Sets if the port should listen to its default queue. 
+	 * If false, another raw listener can create the default queue by adding a raw listener and inject its data via GameThreadInjectDMXSignal
+	 * Useful for nDisplay replication.
+	 */
+	void SetUseDefaultQueue(bool bUse);
+	
 protected:
 	/** Called to set if receive DMX should be enabled */
 	void OnSetReceiveDMXEnabled(bool bEnabled);
 
 	/** According to DMXProtcolSettings, true if DMX should be received */
 	bool bReceiveDMXEnabled;
-	
+
 private:
 	/** The default buffer, which is being read on tick */
-	TQueue<FDMXSignalSharedPtr> TickedBuffer;
+	TQueue<FDMXSignalSharedPtr> DefaultInputQueue;
+
+	/** If true, adds data to the default input queue and fills the UniverseToLatestSignalMap with it */
+	bool bUseDefaultInputQueue = true;
 
 	/** Map of Universe Inputs with their Universes */
 	TMap<int32, TSet<TSharedRef<FDMXTickedUniverseListener>>> LocalUniverseToListenerGroupMap;
 
-	/** Map of latest Singals per local Universe */
-	TMap<int32, FDMXSignalSharedPtr> UniverseToLatestSignalMap;
+	/** Map of latest Singals per extern Universe */
+	TMap<int32, FDMXSignalSharedPtr> ExternUniverseToLatestSignalMap;
 
 	/** Map of raw isteners */
 	TSet<TSharedRef<FDMXRawListener>> RawListeners;
diff --git a/Engine/Plugins/VirtualProduction/DMX/DMXProtocol/Source/DMXProtocol/Public/IO/DMXInputPortConfig.h b/Engine/Plugins/VirtualProduction/DMX/DMXProtocol/Source/DMXProtocol/Public/IO/DMXInputPortConfig.h
--- a/Engine/Plugins/VirtualProduction/DMX/DMXProtocol/Source/DMXProtocol/Public/IO/DMXInputPortConfig.h
+++ b/Engine/Plugins/VirtualProduction/DMX/DMXProtocol/Source/DMXProtocol/Public/IO/DMXInputPortConfig.h
@@ -67,6 +67,13 @@ public:
 	/** Expose the protected PortGuid property name */
 	static FName GetPortGuidPropertyName() { return GET_MEMBER_NAME_CHECKED(FDMXInputPortConfig, PortGuid); }
 
+	/**
+	 * FDMXInputPortConfig relies on its property type customization to initialize the device address (see SDMXLocalAdapterAddressComboBox).
+	 * That alike the device adresses can be sanetized when the struct is loaded and rebuilding ports while ArrayAdd property changes are ongoing can be avoided.
+	 * The function returns true when the IPAdress field is properly initialized.
+	 */
+	bool IsDeviceAddressInitialized() const;
+
 protected:
 	/** Sets a valid port name if the name is empty */
 	void SanetizePortName();
diff --git a/Engine/Plugins/VirtualProduction/DMX/DMXProtocol/Source/DMXProtocol/Public/IO/DMXOutputPort.h b/Engine/Plugins/VirtualProduction/DMX/DMXProtocol/Source/DMXProtocol/Public/IO/DMXOutputPort.h
--- a/Engine/Plugins/VirtualProduction/DMX/DMXProtocol/Source/DMXProtocol/Public/IO/DMXOutputPort.h
+++ b/Engine/Plugins/VirtualProduction/DMX/DMXProtocol/Source/DMXProtocol/Public/IO/DMXOutputPort.h
@@ -45,6 +45,9 @@ struct FDMXOutputPortCommunicationDeterminator
 	/** Determinates if loopback to engine is needed. If true, loopback is needed */
 	FORCEINLINE bool NeedsSendDMX() const { return bSendEnabled && bHasValidSender; }
 
+	/** Determinates if loopback to engine is needed. If true, loopback is needed */
+	FORCEINLINE bool IsSendDMXEnabled() const { return bSendEnabled; }
+
 private:
 	bool bLoopbackToEngine;
 	bool bReceiveEnabled;
@@ -65,16 +68,13 @@ private:
 class DMXPROTOCOL_API FDMXOutputPort
 	: public FDMXPort
 {
-	// Friend Raw Listener so it can add and remove themselves to the port
+	// Friend DMXPortManager so it can create instances and  unregister void instances
 	friend FDMXPortManager;
 
-	// Friend Raw Listener so it can add and remove themselves to the port
+	// Friend Raw Listener so it can add and remove itself to the port
 	friend FDMXRawListener;
 
 protected:
-	/** Creates an output port that is not tied to a specific config. Hidden on purpose, use FDMXPortManager to create instances */
-	static FDMXOutputPortSharedRef Create();
-
 	/** Creates an output port tied to a specific config. Hidden on purpose, use FDMXPortManager to create instances */
 	static FDMXOutputPortSharedRef CreateFromConfig(const FDMXOutputPortConfig& OutputPortConfig);
 
diff --git a/Engine/Plugins/VirtualProduction/DMX/DMXProtocol/Source/DMXProtocol/Public/IO/DMXOutputPortConfig.h b/Engine/Plugins/VirtualProduction/DMX/DMXProtocol/Source/DMXProtocol/Public/IO/DMXOutputPortConfig.h
--- a/Engine/Plugins/VirtualProduction/DMX/DMXProtocol/Source/DMXProtocol/Public/IO/DMXOutputPortConfig.h
+++ b/Engine/Plugins/VirtualProduction/DMX/DMXProtocol/Source/DMXProtocol/Public/IO/DMXOutputPortConfig.h
@@ -78,6 +78,13 @@ public:
 
 	/** Expose the protected PortGuid property name */
 	static FName GetPortGuidPropertyName() { return GET_MEMBER_NAME_CHECKED(FDMXOutputPortConfig, PortGuid); }
+	
+	/**
+	 * FDMXOutputPortConfig relies on its property type customization to initialize the device address (see SDMXLocalAdapterAddressComboBox).
+	 * That alike the device adresses can be sanetized when the struct is loaded and rebuilding ports while ArrayAdd property changes are ongoing can be avoided.
+	 * The function returns true when the IPAdress field is properly initialized.
+	 */
+	bool IsDeviceAddressInitialized() const;
 
 protected:
 	/** Sets a valid port name if the name is empty */
diff --git a/Engine/Plugins/VirtualProduction/DMX/DMXProtocol/Source/DMXProtocol/Public/IO/DMXPortManager.h b/Engine/Plugins/VirtualProduction/DMX/DMXProtocol/Source/DMXProtocol/Public/IO/DMXPortManager.h
--- a/Engine/Plugins/VirtualProduction/DMX/DMXProtocol/Source/DMXProtocol/Public/IO/DMXPortManager.h
+++ b/Engine/Plugins/VirtualProduction/DMX/DMXProtocol/Source/DMXProtocol/Public/IO/DMXPortManager.h
@@ -24,41 +24,38 @@ class FDMXOutputPort;
  * 1. Definition of Ports in Project Settings:
  * -------------------------------------------
  * DMX Protocol Settings (the DMX Project Settings) holds arrays of DMX Port Configs. This is where ports are defined.
- * Port Manager automatically creates DMXInputPorts and DMXOutput ports that match these settings. 
- * Generally this is self-contained and does not need any user code.
  *
+ * Port Manager automatically creates DMXInputPortS and DMXOutputPortS that match these settings. Generally this is self-contained and does not need any user code.
+ *
+ * Note: Applications that want to offer dynamic ports should specify a fixed number of ports in project settings, e.g. 8, 16 ports and work with these at runtime.
+ *		 (While it is possible to mutate the protocol settings at runtime and create ports from that via FDMXPortManager.Get()::UpdateFromProtocolSettings(), 
+ *		 ports created alike are not available in blueprints as port references, hence this approach is disfavored in most situations.)
  * 
- * 2. Acquire a DMX Input or Output Port:
- * -------------------------------------_
- * Get all input ports or all output ports available via the Port Managers GetInputPorts and GetOutputPorts methods.
- * Alternatively use the Editor-Only SDMXPortSelector widget to select a port from available ports.
  * 
- * Note: Creating and destroying ports can only be done in Project settings. 
- *		 Applications that want to offer dynamic ports should specify a fixed number of ports, e.g. 8, 16 ports and work with these at runtime.
+ * 2. Acquire a DMXInputPort or DMXOutputPort:
+ * -------------------------------------------
+ * Get all input ports or all output ports available via the Port Manager's GetInputPorts and GetOutputPorts methods.
+ * Alternatively use the Editor-Only SDMXPortSelector widget to select a port from available ports.
  * 
  *
- * 3. Receive DMX in your object:
+ * 3. Receive DMX in an object:
  * ----------------------------------
  * a) Latest (frame time relevant) data on the Game-Thread:
  * 
- * - Get the port you want to use (see 2. acquire a DMX Input or Output port). GameThreadGetDMXSignal to get a DMX Signal of a local universe.
+ * - Call the port's GameThreadGetDMXSignal method to get a DMX Signal of a local universe.
  * 
  * b) All data on any thread:
  * 
- * - Create an instance of DMXRawListner. Use its constructor to specify which port it should use (see 2. acquire a DMX Input or Output port)
+ * - Create an instance of DMXRawListener. Use its constructor to specify which port it should use.
  * - Use DequeueSignal to receive DMX.
  *
  * Note: This applies for both input and output ports. This is to provide loopback functionality for outputs.
- *		 Generally you want to listen to all ports, not just the inputs. 
+ *		 Generally you want to listen to input and output ports, not just the inputs. 
  *
- * Note: GetDMXData is the right method to use for almost any use-case.
+ * Note: GameThreadGetDMXSignal is the right method to use for almost any use-case.
  *		 DMXRawListener is only useful where the latest data isn't sufficient, e.g. to record all incoming data in Sequencer.
  *		 DMXRawListener is thread-safe, but may stall the engine when used in the game-thread due to the possibly infinite work load it leaves to the user.
  *
- * Note: DMX only runs at a rate of 44Hz by its own Standard, which may be a lower rate than the engine's frame rate.
- *       This may cause visible artifacts. It is recommended to interpolate between received values to overcome the issue.
- *
- *
  * 4. Send DMX from your object:
  * -------------------------------
  * Use the DMXOutputPort's SendDMX method to output DMX
@@ -66,52 +63,23 @@ class FDMXOutputPort;
  */
 class DMXPROTOCOL_API FDMXPortManager
 {
-	////////////////////////////////
-	// Commonly used methods
-
-	DECLARE_MULTICAST_DELEGATE(FDMXOnPortArraysChangedDelegate);
-	DECLARE_MULTICAST_DELEGATE_OneParam(FDMXOnPortChangedDelegate, const FGuid& /** PortGuid */)
-	DECLARE_MULTICAST_DELEGATE_TwoParams(FDMXOnPortInputDequeuedDelegate, FDMXInputPortSharedRef /** DMXInputPortRef */, FDMXSignalSharedRef /** DMXSignalRef */)
+	DECLARE_MULTICAST_DELEGATE(FDMXOnPortsChangedDelegate);
 
 public:
 	/** Broadcast when port arrays or data changed */
-	FDMXOnPortArraysChangedDelegate OnPortsChanged;
+	FDMXOnPortsChangedDelegate OnPortsChanged;
 
-	/** Broadcast when a DMX Signal is dequeued from the port */
-	FDMXOnPortInputDequeuedDelegate OnPortInputDequeued;
-
-public:
 	static FDMXPortManager& Get();
 
 	FORCEINLINE const TArray<FDMXInputPortSharedRef>& GetInputPorts() const { return InputPorts; }
 
 	FORCEINLINE const TArray<FDMXOutputPortSharedRef>& GetOutputPorts() const { return OutputPorts; }
 
-	/** Adds a new input port */
-	FDMXInputPortSharedRef CreateInputPort();
-
-	/** 
-	 * Adds a new input port that corresponds to the input port config. 
-	 * Consumes (moves) the config and returns a copy that uses the port's guid.
-	 * This is to prevent from adding a config (with the same guid) twice.
-	 */
-	FDMXInputPortSharedRef GetOrCreateInputPortFromConfig(const FDMXInputPortConfig& InputPortConfig);
-
-	/** Removes the input port */
-	void RemoveInputPortChecked(const FGuid& PortGuid);
-
-	/** Adds a new output port */
-	FDMXOutputPortSharedRef CreateOutputPort();
-
-	/** 
-	 * Adds a new output port that corresponds to the input port config. 
-	 * Consumes (moves) the config and returns a copy that uses the port's guid 
-	 * This is to prevent from adding a config (with the same guid) twice.
-	 */
-	FDMXOutputPortSharedRef GetOrCreateOutputPortFromConfig(const FDMXOutputPortConfig& OutputPortConfig);
+	/** Gets the input port that corresponds to the input port config. Checks the config is in the DMXProtocolSetting's InputPortConfigs array. */
+	FDMXInputPortSharedRef GetInputPortFromConfigChecked(const FDMXInputPortConfig& InputPortConfig);
 
-	/** Removes the input port */
-	void RemoveOutputPortChecked(const FGuid& PortGuid);
+	/** Gets the output port that corresponds to the input port config. Checks the config is in the DMXProtocolSetting's OutputPortConfigs array. */
+	FDMXOutputPortSharedRef GetOutputPortFromConfigChecked(const FDMXOutputPortConfig& OutputPortConfig);
 
 	/** Returns the port matching the port guid. Returns nullptr if the port doesn't exist. */
 	FDMXPortSharedPtr FindPortByGuid(const FGuid& PortGuid) const;
@@ -131,10 +99,33 @@ public:
 	/** Returns the output port matching the port guid, checked version. */
 	FDMXOutputPortSharedRef FindOutputPortByGuidChecked(const FGuid& PortGuid) const;
 
-	/** Updates ports from protocol settings, does not affect other ports added by the create port methods */
+	/** Updates ports from protocol settings. Useful when project settings were modified in some way */
 	void UpdateFromProtocolSettings();
 
+	/** 
+	 * Thread-safe: Suspends the protocols irrevocably, and with it all sending and receiving of DMX via the protocols. 
+	 * Useful when UE4 self wants to inject data into ports, rather than communicating with the outside world, e.g. for nDisplay replication.
+	 * 
+	 * Note: When protocols are suspended, changes to project settings are no longer reflected and UpdateFromProtocolSettings has no effect.
+	 */
+	void SuspendProtocols();
+
+	/** Returns true if protocols are suspended */
+	bool AreProtocolsSuspended() { return bProtocolsSuspended; }
+
 private:
+	/** Gets the input port that corresponds to the input port config or creates a new one. Note, the config needs be in the DMXProtocolSetting input ports array. */
+	FDMXInputPortSharedRef GetOrCreateInputPortFromConfig(const FDMXInputPortConfig& InputPortConfig);
+
+	/** Removes the input port. Checks that it exists. */
+	void RemoveInputPortChecked(const FGuid& PortGuid);
+
+	/** Gets the output port that corresponds to the input port config or creates a new one. Note, the config needs be in the DMXProtocolSetting input ports array. */
+	FDMXOutputPortSharedRef GetOrCreateOutputPortFromConfig(const FDMXOutputPortConfig& OutputPortConfig);
+
+	/** Removes the output port. Checks that it exists. */
+	void RemoveOutputPortChecked(const FGuid& PortGuid);
+
 	/** Array of input ports */
 	TArray<FDMXInputPortSharedRef> InputPorts;
 
@@ -144,6 +135,9 @@ private:
 	/** Array of Port Guids added from protocol settings */
 	TArray<FGuid> PortGuidsFromProtocolSettings;
 
+	/** True when protocols are suspended */
+	bool bProtocolsSuspended;
+
 	////////////////////////////////////////////////////////////
 	// Initialization 
 public:
