DMX - Rework Output Ports to support many unicast addresses per port, add a selector for the delay type (fps or seconds) and avoid resending unchanged data

diff --git a/Engine/Plugins/VirtualProduction/DMX/DMXEngine/Source/DMXRuntime/Private/Library/DMXLibrary.cpp b/Engine/Plugins/VirtualProduction/DMX/DMXEngine/Source/DMXRuntime/Private/Library/DMXLibrary.cpp
--- a/Engine/Plugins/VirtualProduction/DMX/DMXEngine/Source/DMXRuntime/Private/Library/DMXLibrary.cpp
+++ b/Engine/Plugins/VirtualProduction/DMX/DMXEngine/Source/DMXRuntime/Private/Library/DMXLibrary.cpp
@@ -680,7 +680,6 @@ void UDMXLibrary::UpgradeFromControllersToPorts()
 					OutputPortConfigParams.CommunicationType = VoidController->CommunicationMode;
 					OutputPortConfigParams.ProtocolName = ProtocolName;
 					OutputPortConfigParams.DeviceAddress = InterfaceIPAddress_DEPRECATED;
-					OutputPortConfigParams.DestinationAddress = TEXT("");
 					OutputPortConfigParams.bLoopbackToEngine = OutputPortConfigParams.CommunicationType == EDMXCommunicationType::Multicast ? true : false;
 					OutputPortConfigParams.LocalUniverseStart = FixedLocalUniverseStart;
 					OutputPortConfigParams.NumUniverses = FixedNumUniverses;
@@ -695,7 +694,7 @@ void UDMXLibrary::UpgradeFromControllersToPorts()
 				}
 
 				// Add ports from additional unicast ip
-				for (const FString& AdditionalUnicastIP : VoidController->AdditionalUnicastIPs)
+				if (VoidController->AdditionalUnicastIPs.Num() > 0)
 				{
 					ExistingOutputPortConfigPtr = ProtocolSettings->FindOutputPortConfig([VoidController, &InterfaceIPAddress_DEPRECATED, &ProtocolName](FDMXOutputPortConfig& OutputPortConfig) {
 						return
@@ -704,11 +703,17 @@ void UDMXLibrary::UpgradeFromControllersToPorts()
 							OutputPortConfig.GetCommunicationType() == EDMXCommunicationType::Unicast;
 					});
 
+					TArray<FDMXOutputPortDestinationAddress> DestinationAddressStructs;
+					for (const FString& DestinationAddress : VoidController->AdditionalUnicastIPs)
+					{
+						DestinationAddressStructs.Add(FDMXOutputPortDestinationAddress(DestinationAddress));
+					}
+
 					if (ExistingOutputPortConfigPtr)
 					{
 						FDMXOutputPortConfigParams OutputPortConfigParams(*ExistingOutputPortConfigPtr);
 						OutputPortConfigParams.CommunicationType = EDMXCommunicationType::Unicast;
-						OutputPortConfigParams.DestinationAddress = AdditionalUnicastIP;
+						OutputPortConfigParams.DestinationAddresses = DestinationAddressStructs;
 						OutputPortConfigParams.LocalUniverseStart = FMath::Min(ExistingOutputPortConfigPtr->GetLocalUniverseStart(), FixedLocalUniverseStart);
 						OutputPortConfigParams.NumUniverses = FMath::Max(ExistingOutputPortConfigPtr->GetNumUniverses(), FixedNumUniverses);
 
@@ -723,7 +728,7 @@ void UDMXLibrary::UpgradeFromControllersToPorts()
 						OutputPortConfigParams.CommunicationType = EDMXCommunicationType::Unicast;
 						OutputPortConfigParams.ProtocolName = ProtocolName;
 						OutputPortConfigParams.DeviceAddress = InterfaceIPAddress_DEPRECATED;
-						OutputPortConfigParams.DestinationAddress = AdditionalUnicastIP;
+						OutputPortConfigParams.DestinationAddresses = DestinationAddressStructs;
 						OutputPortConfigParams.bLoopbackToEngine = true;
 						OutputPortConfigParams.LocalUniverseStart = FixedLocalUniverseStart;
 						OutputPortConfigParams.NumUniverses = FixedNumUniverses;
diff --git a/Engine/Plugins/VirtualProduction/DMX/DMXProtocol/Source/DMXProtocol/Private/DMXProtocolBlueprintLibrary.cpp b/Engine/Plugins/VirtualProduction/DMX/DMXProtocol/Source/DMXProtocol/Private/DMXProtocolBlueprintLibrary.cpp
--- a/Engine/Plugins/VirtualProduction/DMX/DMXProtocol/Source/DMXProtocol/Private/DMXProtocolBlueprintLibrary.cpp
+++ b/Engine/Plugins/VirtualProduction/DMX/DMXProtocol/Source/DMXProtocol/Private/DMXProtocolBlueprintLibrary.cpp
@@ -103,8 +103,39 @@ void UDMXProtocolBlueprintLibrary::SetDMXOutputPortDeviceAddress(FDMXOutputPortR
 	}
 }
 
+void UDMXProtocolBlueprintLibrary::SetDMXOutputPortDestinationAddresses(FDMXOutputPortReference OutputPort, const TArray<FString>& DestinationAddresses)
+{
+	UDMXProtocolSettings* ProtocolSettings = GetMutableDefault<UDMXProtocolSettings>();
+
+	if (ProtocolSettings)
+	{
+		FDMXOutputPortConfig* OutputPortConfigPtr = ProtocolSettings->OutputPortConfigs.FindByPredicate([&OutputPort](const FDMXOutputPortConfig& OutputPortConfig)
+			{
+				return OutputPortConfig.GetPortGuid() == OutputPort.GetPortGuid();
+			});
+
+		if (OutputPortConfigPtr)
+		{
+			TArray<FDMXOutputPortDestinationAddress> DestinationAddressesStructs;
+			for (const FString& DestinationAddress : DestinationAddresses)
+			{
+				DestinationAddressesStructs.Add(FDMXOutputPortDestinationAddress(DestinationAddress));
+			}
+
+			FDMXOutputPortConfigParams OutputPortConfigParams(*OutputPortConfigPtr);
+			OutputPortConfigParams.DestinationAddresses = DestinationAddressesStructs;
+
+			*OutputPortConfigPtr = FDMXOutputPortConfig(OutputPortConfigPtr->GetPortGuid(), OutputPortConfigParams);
+
+			FDMXPortManager::Get().UpdateFromProtocolSettings();
+		}
+	}
+}
+
 void UDMXProtocolBlueprintLibrary::SetDMXOutputPortDestinationAddress(FDMXOutputPortReference OutputPort, const FString& DestinationAddress)
 {
+	// DEPRECATED 5.0
+
 	UDMXProtocolSettings* ProtocolSettings = GetMutableDefault<UDMXProtocolSettings>();
 
 	if (ProtocolSettings)
@@ -117,7 +148,15 @@ void UDMXProtocolBlueprintLibrary::SetDMXOutputPortDestinationAddress(FDMXOutput
 		if (OutputPortConfigPtr)
 		{
 			FDMXOutputPortConfigParams OutputPortConfigParams(*OutputPortConfigPtr);
-			OutputPortConfigParams.DestinationAddress = DestinationAddress;
+
+			if (OutputPortConfigParams.DestinationAddresses.Num() > 0)
+			{
+				OutputPortConfigParams.DestinationAddresses[0] = FDMXOutputPortDestinationAddress(DestinationAddress);
+			}
+			else
+			{
+				OutputPortConfigParams.DestinationAddresses.Add(FDMXOutputPortDestinationAddress(DestinationAddress));
+			}
 
 			*OutputPortConfigPtr = FDMXOutputPortConfig(OutputPortConfigPtr->GetPortGuid(), OutputPortConfigParams);
 
diff --git a/Engine/Plugins/VirtualProduction/DMX/DMXProtocol/Source/DMXProtocol/Private/DMXProtocolSettings.cpp b/Engine/Plugins/VirtualProduction/DMX/DMXProtocol/Source/DMXProtocol/Private/DMXProtocolSettings.cpp
--- a/Engine/Plugins/VirtualProduction/DMX/DMXProtocol/Source/DMXProtocol/Private/DMXProtocolSettings.cpp
+++ b/Engine/Plugins/VirtualProduction/DMX/DMXProtocol/Source/DMXProtocol/Private/DMXProtocolSettings.cpp
@@ -4,10 +4,14 @@
 
 #include "DMXProtocolBlueprintLibrary.h"
 #include "DMXProtocolConstants.h"
+#include "DMXProtocolObjectVersion.h"
 #include "IO/DMXPortManager.h"
 #include "Interfaces/IDMXProtocol.h"
 #include "IO/DMXPortManager.h"
 
+#include "IPAddress.h"
+#include "SocketSubsystem.h"
+
 
 UDMXProtocolSettings::UDMXProtocolSettings()
 	: SendingRefreshRate(DMX_RATE)
@@ -96,6 +100,36 @@ void UDMXProtocolSettings::PostInitProperties()
 	OverrideReceiveDMXEnabled(bDefaultReceiveDMXEnabled);	
 }
 
+void UDMXProtocolSettings::PostLoad()
+{
+	Super::PostLoad();
+
+	// Upgrade from single to many destination addresses for ports
+	const int32 CustomVersion = GetLinkerCustomVersion(FDMXProtocolObjectVersion::GUID);
+	if (CustomVersion < FDMXProtocolObjectVersion::OutputPortSupportsManyUnicastAddresses)
+	{
+		ISocketSubsystem* SocketSubsystem = ISocketSubsystem::Get(PLATFORM_SOCKETSUBSYSTEM);
+		TSharedPtr<FInternetAddr> InternetAddr = SocketSubsystem->CreateInternetAddr();
+
+		for (FDMXOutputPortConfig& OutputPortConfig : OutputPortConfigs)
+		{
+			PRAGMA_DISABLE_DEPRECATION_WARNINGS
+			FString DestinationAddress = OutputPortConfig.GetDestinationAddress();
+			PRAGMA_ENABLE_DEPRECATION_WARNINGS
+
+			bool bIsValidIP = false;
+			InternetAddr->SetIp(*DestinationAddress, bIsValidIP);
+			if (bIsValidIP)
+			{
+				FDMXOutputPortConfigParams NewPortParams = FDMXOutputPortConfigParams(OutputPortConfig);
+				NewPortParams.DestinationAddresses = { DestinationAddress };
+
+				OutputPortConfig = FDMXOutputPortConfig(OutputPortConfig.GetPortGuid(), NewPortParams);
+			}
+		}
+	}
+}
+
 #if WITH_EDITOR
 void UDMXProtocolSettings::PostEditChangeProperty(FPropertyChangedEvent& PropertyChangedEvent)
 {
@@ -191,6 +225,12 @@ void UDMXProtocolSettings::PostEditChangeChainProperty(FPropertyChangedChainEven
 
 		FDMXPortManager::Get().UpdateFromProtocolSettings();
 	}
+	else if (
+		PropertyName == FDMXOutputPortConfig::GetDestinationAddressesPropertyNameChecked() ||
+		PropertyName == GET_MEMBER_NAME_CHECKED(FDMXOutputPortDestinationAddress, DestinationAddressString))
+	{
+		FDMXPortManager::Get().UpdateFromProtocolSettings();
+	}
 	
 	Super::PostEditChangeChainProperty(PropertyChangedChainEvent);
 }
diff --git a/Engine/Plugins/VirtualProduction/DMX/DMXProtocol/Source/DMXProtocol/Private/IO/DMXOutputPort.cpp b/Engine/Plugins/VirtualProduction/DMX/DMXProtocol/Source/DMXProtocol/Private/IO/DMXOutputPort.cpp
--- a/Engine/Plugins/VirtualProduction/DMX/DMXProtocol/Source/DMXProtocol/Private/IO/DMXOutputPort.cpp
+++ b/Engine/Plugins/VirtualProduction/DMX/DMXProtocol/Source/DMXProtocol/Private/IO/DMXOutputPort.cpp
@@ -9,10 +9,55 @@
 #include "IO/DMXPortManager.h"
 #include "IO/DMXRawListener.h"
 
+#include "HAL/RunnableThread.h"
 
+
 #define LOCTEXT_NAMESPACE "DMXOutputPort"
 
 
+
+FDMXSignalSharedPtr FDMXThreadSafeUniverseToSignalMap::GetSignal(int32 UniverseID) const
+{
+	const FDMXSignalSharedPtr* SignalPtr = UniverseToSignalMap.Find(UniverseID);
+	
+	return SignalPtr ? *SignalPtr : nullptr;
+}
+
+FDMXSignalSharedRef FDMXThreadSafeUniverseToSignalMap::GetOrCreateSignal(int32 UniverseID)
+{
+	FDMXSignalSharedPtr& Signal = UniverseToSignalMap.FindOrAdd(UniverseID);
+	if (!Signal.IsValid())
+	{
+		Signal = MakeShared<FDMXSignal, ESPMode::ThreadSafe>();
+	}
+
+	return Signal.ToSharedRef();
+}
+
+void FDMXThreadSafeUniverseToSignalMap::AddSignal(const FDMXSignalSharedRef& Signal)
+{
+	FScopeLock Lock(&CriticalSection);
+
+	UniverseToSignalMap.Add(Signal->ExternUniverseID, Signal);
+}
+
+void FDMXThreadSafeUniverseToSignalMap::ForEachSignal(TUniverseSignalPredicate Predicate)
+{
+	FScopeLock Lock(&CriticalSection);
+
+	for (TTuple<int32, FDMXSignalSharedPtr>& UniverseToSignalPair : UniverseToSignalMap)
+	{
+		Predicate(UniverseToSignalPair.Key, UniverseToSignalPair.Value);
+	}
+}
+
+void FDMXThreadSafeUniverseToSignalMap::Reset()
+{
+	FScopeLock Lock(&CriticalSection);
+
+	UniverseToSignalMap.Reset();
+}
+
 FDMXOutputPortSharedRef FDMXOutputPort::CreateFromConfig(FDMXOutputPortConfig& OutputPortConfig)
 {
 	// Port Configs are expected to have a valid guid always
@@ -28,14 +73,14 @@
 	NewOutputPort->CommunicationDeterminator.SetSendEnabled(Settings->IsSendDMXEnabled());
 	NewOutputPort->CommunicationDeterminator.SetReceiveEnabled(Settings->IsReceiveDMXEnabled());
 
-	// Bind to send dmx changes
 	Settings->OnSetSendDMXEnabled.AddThreadSafeSP(NewOutputPort, &FDMXOutputPort::OnSetSendDMXEnabled);
-
-	// Bind to receive dmx changes
 	Settings->OnSetReceiveDMXEnabled.AddThreadSafeSP(NewOutputPort, &FDMXOutputPort::OnSetReceiveDMXEnabled);
 
 	NewOutputPort->UpdateFromConfig(OutputPortConfig);
 
+	const FString SenderThreadName = FString(TEXT("DMXOutputPort_")) + OutputPortConfig.GetPortName();
+	NewOutputPort->Thread = FRunnableThread::Create(&NewOutputPort.Get(), *SenderThreadName, 0U, TPri_TimeCritical, FPlatformAffinity::GetPoolThreadMask());
+
 	UE_LOG(LogDMXProtocol, VeryVerbose, TEXT("Created output port %s"), *NewOutputPort->PortName);
 
 	return NewOutputPort;
@@ -47,7 +92,13 @@
 	check(RawListeners.Num() == 0);
 	
 	// Port needs be unregistered before destruction
-	check(!DMXSender);
+	check(DMXSenderArray.Num() == 0);
 
+	if (Thread)
+	{
+		Thread->Kill(true);
+		delete Thread;
+	}
+
 	UE_LOG(LogDMXProtocol, VeryVerbose, TEXT("Destroyed output port %s"), *PortName);
 }
@@ -76,10 +127,10 @@
 
 		if (ProtocolName == OutputPortConfig.GetProtocolName() &&
 			DeviceAddress == OutputPortConfig.GetDeviceAddress() &&
-			DestinationAddress == OutputPortConfig.GetDestinationAddress() &&
+			DestinationAddresses == OutputPortConfig.GetDestinationAddresses() &&
 			CommunicationType == OutputPortConfig.GetCommunicationType() &&
 			Priority == OutputPortConfig.GetPriority() &&
-			Delay == OutputPortConfig.GetDelay())
+			DelaySeconds == OutputPortConfig.GetDelaySeconds())
 		{
 			return false;
 		}	
@@ -106,12 +157,12 @@
 	CommunicationType = OutputPortConfig.GetCommunicationType();
 	ExternUniverseStart = OutputPortConfig.GetExternUniverseStart();
 	DeviceAddress = OutputPortConfig.GetDeviceAddress();
-	DestinationAddress = OutputPortConfig.GetDestinationAddress();
+	DestinationAddresses = OutputPortConfig.GetDestinationAddresses();
 	LocalUniverseStart = OutputPortConfig.GetLocalUniverseStart();
 	NumUniverses = OutputPortConfig.GetNumUniverses();
 	PortName = OutputPortConfig.GetPortName();
 	Priority = OutputPortConfig.GetPriority();
-	Delay = OutputPortConfig.GetDelay();
+	DelaySeconds = OutputPortConfig.GetDelaySeconds();
 
 	CommunicationDeterminator.SetLoopbackToEngine(OutputPortConfig.NeedsLoopbackToEngine());
 
@@ -130,9 +181,83 @@
 	return PortGuid;
 }
 
+void FDMXOutputPort::ClearBuffers()
+{
+	for (const TSharedRef<FDMXRawListener>& RawListener : RawListeners)
+	{
+		RawListener->ClearBuffer();
+	}
+
+	LatestDMXSignals.Reset();
+}
+
+bool FDMXOutputPort::IsLoopbackToEngine() const
+{
+	return CommunicationDeterminator.NeedsLoopbackToEngine();
+}
+
+TArray<FString> FDMXOutputPort::GetDestinationAddresses() const
+{
+	TArray<FString> Result;
+	Result.Reserve(DestinationAddresses.Num());
+	for (const FDMXOutputPortDestinationAddress& DestinationAddress : DestinationAddresses)
+	{
+		Result.Add(DestinationAddress.DestinationAddressString);
+	}
+
+	return Result;
+}
+
+bool FDMXOutputPort::GameThreadGetDMXSignal(int32 LocalUniverseID, FDMXSignalSharedPtr& OutDMXSignal, bool bEvenIfNotLoopbackToEngine)
+{
+#if UE_BUILD_DEBUG
+	check(IsInGameThread());
+#endif // UE_BUILD_DEBUG
+
+	const bool bNeedsLoopbackToEngine = CommunicationDeterminator.NeedsLoopbackToEngine();
+	if (bNeedsLoopbackToEngine || bEvenIfNotLoopbackToEngine)
+	{
+		int32 ExternUniverseID = ConvertLocalToExternUniverseID(LocalUniverseID);
+
+		OutDMXSignal = LatestDMXSignals.GetSignal(ExternUniverseID);
+		if (OutDMXSignal.IsValid())
+		{
+			return true;
+		}
+	}
+
+	return false;
+}
+
+bool FDMXOutputPort::GameThreadGetDMXSignalFromRemoteUniverse(FDMXSignalSharedPtr& OutDMXSignal, int32 RemoteUniverseID, bool bEvenIfNotLoopbackToEngine)
+{
+	// DEPRECATED 4.27
+#if UE_BUILD_DEBUG
+	check(IsInGameThread());
+#endif // UE_BUILD_DEBUG
+
+	const bool bNeedsLoopbackToEngine = CommunicationDeterminator.NeedsLoopbackToEngine();
+	if (bNeedsLoopbackToEngine || bEvenIfNotLoopbackToEngine)
+	{
+		OutDMXSignal = LatestDMXSignals.GetSignal(RemoteUniverseID);
+		if (OutDMXSignal.IsValid())
+		{
+			return true;
+		}
+	}
+
+	return false;
+}
+
+FString FDMXOutputPort::GetDestinationAddress() const
+{
+	// DEPRECATED 5.0
+	return DestinationAddresses.Num() > 0 ? DestinationAddresses[0].DestinationAddressString : TEXT("");
+}
+
 bool FDMXOutputPort::IsRegistered() const
 {
-	if (DMXSender)
+	if (DMXSenderArray.Num() > 0)
 	{
 		return true;
 	}
@@ -167,15 +292,8 @@
 		{
 			const int32 ExternUniverseID = ConvertLocalToExternUniverseID(LocalUniverseID);
 
-			FDMXSignalSharedPtr& ExistingOrNewSignal = ExternUniverseToLatestSignalMap.FindOrAdd(ExternUniverseID);
-			if (!ExistingOrNewSignal.IsValid())
-			{
-				// Only create a new signal, initialization happens below
-				ExistingOrNewSignal = MakeShared<FDMXSignal, ESPMode::ThreadSafe>();
-			}
-			FDMXSignalSharedRef Signal = ExistingOrNewSignal.ToSharedRef();
+			FDMXSignalSharedRef Signal = LatestDMXSignals.GetOrCreateSignal(ExternUniverseID);
 
-			// Init or update the signal
 			// Write the fragment & meta data 
 			for (const TTuple<int32, uint8>& ChannelValueKvp : ChannelToValueMap)
 			{
@@ -189,21 +307,10 @@
 			}	
 
 			Signal->ExternUniverseID = ExternUniverseID;
-			Signal->Timestamp = FPlatformTime::Seconds();
+			Signal->Timestamp = FPlatformTime::Seconds() + DelaySeconds;
 			Signal->Priority = Priority;
-			Signal->Delay = Delay;
 
-			// Send DMX
-			if (bNeedsSendDMX)
-			{
-				DMXSender->SendDMXSignal(Signal);
-			}
-
-			// Loopback to Listeners
-			for (const TSharedRef<FDMXRawListener>& RawListener : RawListeners)
-			{
-				RawListener->EnqueueSignal(this, Signal);
-			}
+			NewDMXSignalsToSend.Enqueue(Signal);
 		}
 	}
 }
@@ -219,15 +326,8 @@
 		// Update the buffer for loopback if dmx needs be sent and/or looped back
 		if (bNeedsSendDMX || bNeedsLoopbackToEngine)
 		{
-			FDMXSignalSharedPtr& ExistingOrNewSignal = ExternUniverseToLatestSignalMap.FindOrAdd(RemoteUniverse);
-			if (!ExistingOrNewSignal.IsValid())
-			{
-				// Only create a new signal, initialization happens below
-				ExistingOrNewSignal = MakeShared<FDMXSignal, ESPMode::ThreadSafe>();
-			}
-			FDMXSignalSharedRef Signal = ExistingOrNewSignal.ToSharedRef();
+			FDMXSignalSharedRef Signal = LatestDMXSignals.GetOrCreateSignal(RemoteUniverse);
 
-			// Init or update the signal
 			// Write the fragment & meta data 
 			for (const TTuple<int32, uint8>& ChannelValueKvp : ChannelToValueMap)
 			{
@@ -240,32 +340,21 @@
 			}
 
 			Signal->ExternUniverseID = RemoteUniverse;
-			Signal->Timestamp = FPlatformTime::Seconds();
+			Signal->Timestamp = FPlatformTime::Seconds() + DelaySeconds;
 			Signal->Priority = Priority;
-			Signal->Delay = Delay;
 
-			// Send via the protocol's sender
-			if (bNeedsSendDMX)
-			{
-				DMXSender->SendDMXSignal(Signal);
-			}
-
-			// Loopback to Listeners
-				for (const TSharedRef<FDMXRawListener>& RawListener : RawListeners)
-				{
-					RawListener->EnqueueSignal(this, Signal);
-				}
-			}
+			NewDMXSignalsToSend.Enqueue(Signal);
 		}
 	}
+}
 
 bool FDMXOutputPort::Register()
 {
 	if (Protocol.IsValid() && IsValidPortSlow() && CommunicationDeterminator.IsSendDMXEnabled() && !FDMXPortManager::Get().AreProtocolsSuspended())
 	{
-		DMXSender = Protocol->RegisterOutputPort(SharedThis(this));
+		DMXSenderArray = Protocol->RegisterOutputPort(SharedThis(this));
 
-		if (DMXSender.IsValid())
+		if (DMXSenderArray.Num() > 0)
 		{
 			CommunicationDeterminator.SetHasValidSender(true);
 			return true;
@@ -282,78 +371,15 @@
 	{
 		if (Protocol.IsValid())
 		{
-		Protocol->UnregisterOutputPort(SharedThis(this));
+			Protocol->UnregisterOutputPort(SharedThis(this));
 		}
 
-		DMXSender = nullptr;
+		DMXSenderArray.Reset();
 	}
 
 	CommunicationDeterminator.SetHasValidSender(false);
 }
 
-void FDMXOutputPort::ClearBuffers()
-{
-	if (DMXSender)
-	{
-		DMXSender->ClearBuffer();
-	}
-
-	for (const TSharedRef<FDMXRawListener>& RawListener : RawListeners)
-	{
-		RawListener->ClearBuffer();
-	}
-
-	ExternUniverseToLatestSignalMap.Reset();
-}
-
-bool FDMXOutputPort::IsLoopbackToEngine() const
-{
-	return CommunicationDeterminator.NeedsLoopbackToEngine();
-}
-
-bool FDMXOutputPort::GameThreadGetDMXSignal(int32 LocalUniverseID, FDMXSignalSharedPtr& OutDMXSignal, bool bEvenIfNotLoopbackToEngine)
-{
-#if UE_BUILD_DEBUG
-	check(IsInGameThread());
-#endif // UE_BUILD_DEBUG
-
-	const bool bNeedsLoopbackToEngine = CommunicationDeterminator.NeedsLoopbackToEngine();
-	if (bNeedsLoopbackToEngine || bEvenIfNotLoopbackToEngine)
-	{
-		int32 ExternUniverseID = ConvertLocalToExternUniverseID(LocalUniverseID);
-
-		const FDMXSignalSharedPtr* SignalPtr = ExternUniverseToLatestSignalMap.Find(ExternUniverseID);
-		if (SignalPtr)
-		{
-			OutDMXSignal = *SignalPtr;
-			return true;
-		}
-	}
-
-	return false;
-}
-
-bool FDMXOutputPort::GameThreadGetDMXSignalFromRemoteUniverse(FDMXSignalSharedPtr& OutDMXSignal, int32 RemoteUniverseID, bool bEvenIfNotLoopbackToEngine)
-{
-	// DEPRECATED 4.27
-#if UE_BUILD_DEBUG
-	check(IsInGameThread());
-#endif // UE_BUILD_DEBUG
-
-	const bool bNeedsLoopbackToEngine = CommunicationDeterminator.NeedsLoopbackToEngine();
-	if (bNeedsLoopbackToEngine || bEvenIfNotLoopbackToEngine)
-	{
-		const FDMXSignalSharedPtr* SignalPtr = ExternUniverseToLatestSignalMap.Find(RemoteUniverseID);
-		if (SignalPtr)
-		{
-			OutDMXSignal = *SignalPtr;
-			return true;
-		}
-	}
-
-	return false;
-}
-
 void FDMXOutputPort::OnSetSendDMXEnabled(bool bEnabled)
 {
 	CommunicationDeterminator.SetSendEnabled(bEnabled);
@@ -386,4 +412,110 @@
 	return nullptr;
 }
 
+bool FDMXOutputPort::Init()
+{
+	return true;
+}
+
+uint32 FDMXOutputPort::Run()
+{
+	const UDMXProtocolSettings* DMXSettings = GetDefault<UDMXProtocolSettings>();
+	check(DMXSettings);
+	
+	// Fixed rate delta time
+	const double SendDeltaTime = 1.f / DMXSettings->SendingRefreshRate;
+
+	while (!bStopping)
+	{
+		const double StartTime = FPlatformTime::Seconds();
+
+		ProcessSendDMX();
+
+		const double EndTime = FPlatformTime::Seconds();
+		const double WaitTime = SendDeltaTime - (EndTime - StartTime);
+
+		if (WaitTime > 0.f)
+		{
+			// Sleep by the amount which is set in refresh rate
+			FPlatformProcess::SleepNoStats(WaitTime);
+		}
+
+		// In the unlikely case we took to long to send, we instantly continue, but do not take 
+		// further measures to compensate - We would have to run faster than DMX send rate to catch up.
+	}
+
+	return 0;
+}
+
+void FDMXOutputPort::Stop()
+{
+	bStopping = true;
+}
+
+void FDMXOutputPort::Exit()
+{
+}
+
+void FDMXOutputPort::Tick()
+{
+	ProcessSendDMX();
+}
+
+FSingleThreadRunnable* FDMXOutputPort::GetSingleThreadInterface()
+{
+	return this;
+}
+
+void FDMXOutputPort::ProcessSendDMX()
+{
+	// Delay signals
+	const double Now = FPlatformTime::Seconds();
+
+	// Acquire new DMX signals
+	for (;;)
+	{
+		FDMXSignalSharedPtr OldestDMXSignal;
+		if (NewDMXSignalsToSend.Peek(OldestDMXSignal))
+		{
+			if (OldestDMXSignal->Timestamp <= Now)
+			{
+				LatestDMXSignals.AddSignal(OldestDMXSignal.ToSharedRef());
+
+				NewDMXSignalsToSend.Pop();
+
+				continue;
+			}
+
+			break;
+		}
+	}
+
+	// Send new and alive DMX Signals
+	const bool bNeedsSendDMX = CommunicationDeterminator.NeedsSendDMX();
+	LatestDMXSignals.ForEachSignal([Now, bNeedsSendDMX, this](int32 UniverseID, const FDMXSignalSharedPtr& Signal)
+		{
+			if (Signal->Timestamp <= Now)
+			{
+				// Keeping the signal alive here:
+				// Increment the timestamp by one second so the signal will be sent in one second anew.
+				Signal->Timestamp = Now + 1000.0;
+
+				// Send via the protocol's sender
+				if (bNeedsSendDMX)
+				{
+					for (const TSharedPtr<IDMXSender>& DMXSender : DMXSenderArray)
+					{
+						DMXSender->SendDMXSignal(Signal.ToSharedRef());
+					}
+				}
+
+				// Loopback to Listeners
+				for (const TSharedRef<FDMXRawListener>& RawListener : RawListeners)
+				{
+					RawListener->EnqueueSignal(this, Signal.ToSharedRef());
+				}
+			}
+		});
+}
+
 #undef LOCTEXT_NAMESPACE
diff --git a/Engine/Plugins/VirtualProduction/DMX/DMXProtocol/Source/DMXProtocol/Private/IO/DMXOutputPortConfig.cpp b/Engine/Plugins/VirtualProduction/DMX/DMXProtocol/Source/DMXProtocol/Private/IO/DMXOutputPortConfig.cpp
--- a/Engine/Plugins/VirtualProduction/DMX/DMXProtocol/Source/DMXProtocol/Private/IO/DMXOutputPortConfig.cpp
+++ b/Engine/Plugins/VirtualProduction/DMX/DMXProtocol/Source/DMXProtocol/Private/IO/DMXOutputPortConfig.cpp
@@ -15,22 +15,24 @@ FDMXOutputPortConfigParams::FDMXOutputPortConfigParams(const FDMXOutputPortConfi
 	, ProtocolName(OutputPortConfig.GetProtocolName())
 	, CommunicationType(OutputPortConfig.GetCommunicationType())
 	, DeviceAddress(OutputPortConfig.GetDeviceAddress())
-	, DestinationAddress(OutputPortConfig.GetDestinationAddress())
+	, DestinationAddresses(OutputPortConfig.GetDestinationAddresses())
 	, bLoopbackToEngine(OutputPortConfig.NeedsLoopbackToEngine())
 	, LocalUniverseStart(OutputPortConfig.GetLocalUniverseStart())
 	, NumUniverses(OutputPortConfig.GetNumUniverses())
 	, ExternUniverseStart(OutputPortConfig.GetExternUniverseStart())
 	, Priority(OutputPortConfig.GetPriority())
-	, Delay(OutputPortConfig.GetDelay())
+	, DelaySeconds(OutputPortConfig.GetDelaySeconds())
 {}
 
 
 FDMXOutputPortConfig::FDMXOutputPortConfig()
-	: PortGuid(FGuid::NewGuid())
+	: DelayFrameRate(FFrameRate(1.0, 1.0)) // Default delay frame rate to 1.0 (default to Seconds)
+	, PortGuid(FGuid::NewGuid())
 {}
 
 FDMXOutputPortConfig::FDMXOutputPortConfig(const FGuid& InPortGuid)
-	: PortGuid(InPortGuid)
+	: DelayFrameRate(FFrameRate(1.0, 1.0)) // Default delay frame rate to 1.0 (default to Seconds)
+	, PortGuid(InPortGuid)
 {
 	// Cannot create port configs before the protocol module is up (it is required to sanetize protocol names).
 	check(FModuleManager::Get().IsModuleLoaded("DMXProtocol"));
@@ -45,14 +47,14 @@ FDMXOutputPortConfig::FDMXOutputPortConfig(const FGuid& InPortGuid, const FDMXOu
 	: PortName(InitializationData.PortName)
 	, ProtocolName(InitializationData.ProtocolName)
 	, CommunicationType(InitializationData.CommunicationType)
-	, DeviceAddress(InitializationData.DeviceAddress)
-	, DestinationAddress(InitializationData.DestinationAddress)
+	, DestinationAddresses(InitializationData.DestinationAddresses)
 	, bLoopbackToEngine(InitializationData.bLoopbackToEngine)
 	, LocalUniverseStart(InitializationData.LocalUniverseStart)
 	, NumUniverses(InitializationData.NumUniverses)
 	, ExternUniverseStart(InitializationData.ExternUniverseStart)
 	, Priority(InitializationData.Priority)
-	, Delay(InitializationData.Delay)
+	, Delay(InitializationData.DelaySeconds)
+	, DelayFrameRate(FFrameRate(1.0, 1.0))
 	, PortGuid(InPortGuid)
 {
 	// Cannot create port configs before the protocol module is up (it is required to sanetize protocol names).
@@ -94,37 +96,48 @@ void FDMXOutputPortConfig::MakeValid()
 
 	if (Protocol.IsValid())
 	{
-	// If the extern universe ID is out of the protocol's supported range, mend it.
-	ExternUniverseStart = Protocol->MakeValidUniverseID(ExternUniverseStart);
-
-	// Only Local universes > 1 are supported, even if the protocol supports universes < 1.
-	LocalUniverseStart = LocalUniverseStart < 1 ? 1 : LocalUniverseStart;
+		// If the extern universe ID is out of the protocol's supported range, mend it.
+		ExternUniverseStart = Protocol->MakeValidUniverseID(ExternUniverseStart);
 
-	// Limit the num universes relatively to the max extern universe of the protocol and int32 max
-	const int64 MaxNumUniverses = Protocol->GetMaxUniverseID() - Protocol->GetMinUniverseID() + 1;
-	const int64 DesiredNumUniverses = NumUniverses > MaxNumUniverses ? MaxNumUniverses : NumUniverses;
-	const int64 DesiredLocalUniverseEnd = static_cast<int64>(LocalUniverseStart) + DesiredNumUniverses - 1;
+		// Only Local universes > 1 are supported, even if the protocol supports universes < 1.
+		LocalUniverseStart = LocalUniverseStart < 1 ? 1 : LocalUniverseStart;
 
-	if (DesiredLocalUniverseEnd > TNumericLimits<int32>::Max())
-	{
-		NumUniverses = TNumericLimits<int32>::Max() - DesiredLocalUniverseEnd;
-	}
+		// Limit the num universes relatively to the max extern universe of the protocol and int32 max
+		const int64 MaxNumUniverses = Protocol->GetMaxUniverseID() - Protocol->GetMinUniverseID() + 1;
+		const int64 DesiredNumUniverses = NumUniverses > MaxNumUniverses ? MaxNumUniverses : NumUniverses;
+		const int64 DesiredLocalUniverseEnd = static_cast<int64>(LocalUniverseStart) + DesiredNumUniverses - 1;
 
-	// Fix the communication type if it is not supported by the protocol
-	TArray<EDMXCommunicationType> CommunicationTypes = Protocol->GetOutputPortCommunicationTypes();
-	if (!CommunicationTypes.Contains(CommunicationType))
-	{
-		if (CommunicationTypes.Num() > 0)
+		if (DesiredLocalUniverseEnd > TNumericLimits<int32>::Max())
 		{
-			CommunicationType = CommunicationTypes[0];
+			NumUniverses = TNumericLimits<int32>::Max() - DesiredLocalUniverseEnd;
 		}
-		else
+
+		// Fix the communication type if it is not supported by the protocol
+		TArray<EDMXCommunicationType> CommunicationTypes = Protocol->GetOutputPortCommunicationTypes();
+		if (!CommunicationTypes.Contains(CommunicationType))
 		{
-			// The protocol can specify none to suggest internal only
-			CommunicationType = EDMXCommunicationType::InternalOnly;
+			if (CommunicationTypes.Num() > 0)
+			{
+				CommunicationType = CommunicationTypes[0];
+			}
+			else
+			{
+				// The protocol can specify none to suggest internal only
+				CommunicationType = EDMXCommunicationType::InternalOnly;
+			}
 		}
-	}	
+	}
+
+	// Allow for postitive delay values only
+	if (Delay < 0.0)
+	{
+		Delay = 0.0;
+	}
 }
+
+int32 FDMXOutputPortConfig::GetDelaySeconds() const
+{
+	return Delay / DelayFrameRate.AsDecimal();
 }
 
 FString FDMXOutputPortConfig::GetDeviceAddress() const
@@ -164,4 +177,3 @@ void FDMXOutputPortConfig::GenerateUniquePortName()
 	FString BaseName = TEXT("OutputPort_1");
 	PortName = FDMXProtocolUtils::GenerateUniqueNameFromExisting(OtherPortNames, BaseName);
 }
-
diff --git a/Engine/Plugins/VirtualProduction/DMX/DMXProtocol/Source/DMXProtocol/Public/DMXProtocolBlueprintLibrary.h b/Engine/Plugins/VirtualProduction/DMX/DMXProtocol/Source/DMXProtocol/Public/DMXProtocolBlueprintLibrary.h
--- a/Engine/Plugins/VirtualProduction/DMX/DMXProtocol/Source/DMXProtocol/Public/DMXProtocolBlueprintLibrary.h
+++ b/Engine/Plugins/VirtualProduction/DMX/DMXProtocol/Source/DMXProtocol/Public/DMXProtocolBlueprintLibrary.h
@@ -76,5 +76,14 @@ public:
 	 * @param DestinationAddress	The Destination Address the Output Port should use 
 	 */
 	UFUNCTION(BlueprintCallable, Category = "DMX")
+	static void	SetDMXOutputPortDestinationAddresses(FDMXOutputPortReference OutputPort, const TArray<FString>& DestinationAddresses);
+
+	/**
+	 * Sets the Destination Address Address of the Output Port. For networking Protocols that's the Unicast IP Adress. Not required for Multicast and Broadcast.
+	 * @param PortName				The Output Port for which the Unicast IP Address should be set
+	 * @param DestinationAddress	The Destination Address the Output Port should use 
+	 */
+	UE_DEPRECATED(5.0, "Output Ports now support many destination addresses. Please use UDMXProtocolBlueprintLibrary::SetDMXOutputPortDestinationAddresses instead")
+	UFUNCTION(BlueprintCallable, Category = "DMX", meta = (DeprecatedFunction, DeprecationMessage = "Deprecated 5.0. Output Ports now support many Destination Addresses. Please use SetDMXOutputPortDestinationAddresses instead."))
 	static void	SetDMXOutputPortDestinationAddress(FDMXOutputPortReference OutputPort, const FString& DestinationAddress);
 };
diff --git a/Engine/Plugins/VirtualProduction/DMX/DMXProtocol/Source/DMXProtocol/Public/DMXProtocolObjectVersion.cpp b/Engine/Plugins/VirtualProduction/DMX/DMXProtocol/Source/DMXProtocol/Public/DMXProtocolObjectVersion.cpp
new file mode 100644
--- /dev/null
+++ b/Engine/Plugins/VirtualProduction/DMX/DMXProtocol/Source/DMXProtocol/Public/DMXProtocolObjectVersion.cpp
@@ -0,0 +1,10 @@
+// Copyright Epic Games, Inc. All Rights Reserved.
+
+#include "DMXProtocolObjectVersion.h"
+
+#include "Serialization/CustomVersion.h"
+
+const FGuid FDMXProtocolObjectVersion::GUID(0x94662356, 0x4DAE92AC, 0xABF8BE98, 0x19FE8C2F);
+
+// Register the custom version with core
+FCustomVersionRegistration GRegisterDMXRuntimeCustomVersion(FDMXProtocolObjectVersion::GUID, FDMXProtocolObjectVersion::LatestVersion, TEXT("DMXProtocolObjectVersion"));
diff --git a/Engine/Plugins/VirtualProduction/DMX/DMXProtocol/Source/DMXProtocol/Public/DMXProtocolObjectVersion.h b/Engine/Plugins/VirtualProduction/DMX/DMXProtocol/Source/DMXProtocol/Public/DMXProtocolObjectVersion.h
new file mode 100644
--- /dev/null
+++ b/Engine/Plugins/VirtualProduction/DMX/DMXProtocol/Source/DMXProtocol/Public/DMXProtocolObjectVersion.h
@@ -0,0 +1,28 @@
+// Copyright Epic Games, Inc. All Rights Reserved.
+#pragma once
+
+#include "CoreTypes.h"
+#include "Misc/Guid.h"
+
+// Custom serialization version for changes to DMX Protocol Objects
+struct FDMXProtocolObjectVersion
+{
+	enum Type
+	{
+		// Roughly corresponds to 4.27
+		BeforeCustomVersionWasAdded = 0,
+
+		// Update Ports to allow for an array of Unicast Addresses instead of just one
+		OutputPortSupportsManyUnicastAddresses,
+
+		// -----<new versions can be added above this line>-------------------------------------------------
+		VersionPlusOne,
+		LatestVersion = VersionPlusOne - 1
+	};
+
+	// The GUID for this custom version number
+	const static FGuid GUID;
+
+private:
+	FDMXProtocolObjectVersion() {}
+};
diff --git a/Engine/Plugins/VirtualProduction/DMX/DMXProtocol/Source/DMXProtocol/Public/DMXProtocolSettings.h b/Engine/Plugins/VirtualProduction/DMX/DMXProtocol/Source/DMXProtocol/Public/DMXProtocolSettings.h
--- a/Engine/Plugins/VirtualProduction/DMX/DMXProtocol/Source/DMXProtocol/Public/DMXProtocolSettings.h
+++ b/Engine/Plugins/VirtualProduction/DMX/DMXProtocol/Source/DMXProtocol/Public/DMXProtocolSettings.h
@@ -35,7 +35,7 @@ public:
 
 	// ~Begin UObject Interface
 	virtual void PostInitProperties() override;
-
+	virtual void PostLoad() override;
 #if WITH_EDITOR
 	virtual void PostEditChangeProperty(FPropertyChangedEvent& PropertyChangedEvent) override;
 	virtual void PostEditChangeChainProperty(FPropertyChangedChainEvent& PropertyChangedChainEvent) override;
diff --git a/Engine/Plugins/VirtualProduction/DMX/DMXProtocol/Source/DMXProtocol/Public/DMXProtocolTypes.h b/Engine/Plugins/VirtualProduction/DMX/DMXProtocol/Source/DMXProtocol/Public/DMXProtocolTypes.h
--- a/Engine/Plugins/VirtualProduction/DMX/DMXProtocol/Source/DMXProtocol/Public/DMXProtocolTypes.h
+++ b/Engine/Plugins/VirtualProduction/DMX/DMXProtocol/Source/DMXProtocol/Public/DMXProtocolTypes.h
@@ -37,24 +37,21 @@ public:
 		, ExternUniverseID(0)
 		, Priority(0)
 		, ChannelData()
-		, Delay(0)
 	{
 		ChannelData.AddZeroed(DMX_UNIVERSE_SIZE);
 	}
 
-	FDMXSignal(const double InTimestamp, const int32 InUniverseID, const int32 InPriority, const TArray<uint8>& InChannelData, const float InDelay)
+	FDMXSignal(const double InTimestamp, const int32 InUniverseID, const int32 InPriority, const TArray<uint8>& InChannelData)
 		: Timestamp(InTimestamp)
 		, ExternUniverseID(InUniverseID)
 		, Priority(InPriority)
 		, ChannelData(InChannelData)
-		, Delay(InDelay)
 	{}
 
-	FDMXSignal(const double InTimestamp, const int32 InUniverseID, const int32 InPriority, TArray<uint8>&& InChannelData, const float InDelay)
+	FDMXSignal(const double InTimestamp, const int32 InUniverseID, const int32 InPriority, TArray<uint8>&& InChannelData)
 		: Timestamp(InTimestamp)
 		, ExternUniverseID(InUniverseID)
 		, ChannelData(InChannelData)
-		, Delay(InDelay)
 	{}
 
 	void Serialize(FArchive& Ar)
@@ -63,7 +60,6 @@ public:
 		Ar << ExternUniverseID;
 		Ar << Priority;
 		Ar << ChannelData;
-		Ar << Delay;
 	}
 
 	double Timestamp;
@@ -73,8 +69,6 @@ public:
 	int32 Priority;
 
 	TArray<uint8> ChannelData;
-
-	float Delay;
 };
 
 /** Result when sending a DMX packet */
diff --git a/Engine/Plugins/VirtualProduction/DMX/DMXProtocol/Source/DMXProtocol/Public/IO/DMXOutputPort.h b/Engine/Plugins/VirtualProduction/DMX/DMXProtocol/Source/DMXProtocol/Public/IO/DMXOutputPort.h
--- a/Engine/Plugins/VirtualProduction/DMX/DMXProtocol/Source/DMXProtocol/Public/IO/DMXOutputPort.h
+++ b/Engine/Plugins/VirtualProduction/DMX/DMXProtocol/Source/DMXProtocol/Public/IO/DMXOutputPort.h
@@ -9,9 +9,13 @@
 #include "CoreMinimal.h"
 #include "Tickable.h"
 #include "Containers/Queue.h" 
-#include "Templates/SharedPointer.h"
+#include "HAL/Runnable.h"
+#include "Misc/ScopeLock.h" 
+#include "Misc/SingleThreadRunnable.h"
+#include "Templates/Atomic.h"
 
 struct FDMXOutputPortConfig;
+struct FDMXOutputPortDestinationAddress;
 class FDMXPortManager;
 class FDMXSignal;
 class IDMXSender;
@@ -56,6 +60,35 @@ private:
 };
 
 
+/** Helper to access a Map of Signals per Universe in a thread-safe way */
+class FDMXThreadSafeUniverseToSignalMap
+{
+public:
+	using TUniverseSignalPredicate = TFunctionRef<void(int32, const FDMXSignalSharedPtr&)>;
+	
+	/** Gets a signal, returns an invalid shared pointer if there is no signal at given Universe ID */
+	FDMXSignalSharedPtr GetSignal(int32 UniverseID) const;
+
+	/** Gets or creates a signal */
+	FDMXSignalSharedRef GetOrCreateSignal(int32 UniverseID);
+
+	/** Adds a signal to the map */
+	void AddSignal(const FDMXSignalSharedRef& Signal);
+
+	/** Loops through all signals by given Predicate  */
+	void ForEachSignal(TUniverseSignalPredicate Predicate);
+
+	/** Resets the Map */
+	void Reset();
+
+private:
+	/** Buffer of all the latest DMX Signals that were sent */
+	TMap<int32, FDMXSignalSharedPtr> UniverseToSignalMap;
+
+	/** Critical sequestion required to be used when the Map is accessed */
+	FCriticalSection CriticalSection;
+};
+
 /**
  * Higher level abstraction of a DMX input hiding networking specific and protocol specific complexity.
  *
@@ -67,6 +100,8 @@ private:
  */
 class DMXPROTOCOL_API FDMXOutputPort
 	: public FDMXPort
+	, public FRunnable
+	, public FSingleThreadRunnable
 {
 	// Friend DMXPortManager so it can create instances and unregister void instances
 	friend FDMXPortManager;
@@ -84,17 +119,6 @@ public:
 	/** Updates the Port to use the config of the OutputPortConfig */
 	void UpdateFromConfig(FDMXOutputPortConfig& OutputPortConfig);
 
-	//~ Begin DMXPort Interface 
-	virtual bool IsRegistered() const override;
-	virtual const FGuid& GetPortGuid() const override;
-protected:
-	virtual void AddRawListener(TSharedRef<FDMXRawListener> InRawListener) override;
-	virtual void RemoveRawListener(TSharedRef<FDMXRawListener> InRawListenerToRemove) override;
-	virtual bool Register() override;
-	virtual void Unregister() override;
-	//~ End DMXPort Interface
-
-public:
 	/** Sends DMX over the port */
 	void SendDMX(int32 LocalUniverseID, const TMap<int32, uint8>& ChannelToValueMap);
 
@@ -118,15 +142,30 @@ public:
 	 */
 	bool GameThreadGetDMXSignal(int32 LocalUniverseID, FDMXSignalSharedPtr& OutDMXSignal, bool bEvenIfNotLoopbackToEngine);
 
+	/** Returns the Destination Addresses */
+	TArray<FString> GetDestinationAddresses() const;
+
 	/**  DEPRECATED 4.27. Gets the DMX signal from an extern (remote) Universe ID. */
 	UE_DEPRECATED(4.27, "Use GameThreadGetDMXSignal instead. GameThreadGetDMXSignalFromRemoteUniverse only exists to support deprecated blueprint nodes.")
 	bool GameThreadGetDMXSignalFromRemoteUniverse(FDMXSignalSharedPtr& OutDMXSignal, int32 RemoteUniverseID, bool bEvenIfNotLoopbackToEngine);
 
-	/** Returns the Destination Address */
-	FORCEINLINE const FString& GetDestinationAddress() const { return DestinationAddress; }
+	/** DEPRECATED 5.0 */
+	UE_DEPRECATED(5.0, "Output Ports now support many destination addresses. Please use FDMXOutputPort::GetDestinationAddresses instead.")
+	FString GetDestinationAddress() const;
 
-private:
-	/** Called to set if DMX should be enabled */
+public:
+	//~ Begin DMXPort Interface 
+	virtual bool IsRegistered() const override;
+	virtual const FGuid& GetPortGuid() const override;
+
+protected:
+	virtual void AddRawListener(TSharedRef<FDMXRawListener> InRawListener) override;
+	virtual void RemoveRawListener(TSharedRef<FDMXRawListener> InRawListenerToRemove) override;
+	virtual bool Register() override;
+	virtual void Unregister() override;
+	//~ End DMXPort Interface
+
+		/** Called to set if DMX should be enabled */
 	void OnSetSendDMXEnabled(bool bEnabled);
 
 	/** Called to set if DMX should be enabled */
@@ -134,21 +173,40 @@ private:
 
 	/** Returns the port config that corresponds to the guid of this port. */
 	FDMXOutputPortConfig* FindOutputPortConfigChecked() const;
-	
-	/** The DMX sender, or nullptr if not registered */
-	TSharedPtr<IDMXSender> DMXSender;
+
+	//~ Begin FRunnable implementation
+	virtual bool Init() override;
+	virtual uint32 Run() override;
+	virtual void Stop() override;
+	virtual void Exit() override;
+	//~ End FRunnable implementation
+
+	//~ Begin FSingleThreadRunnable implementation
+	virtual void Tick() override;
+	virtual class FSingleThreadRunnable* GetSingleThreadInterface() override;
+	//~ End FSingleThreadRunnable implementation
+
+	/** Updates the thread, sending DMX */
+	void ProcessSendDMX();
+
+private:
+	/** The DMX senders in use */
+	TArray<TSharedPtr<IDMXSender>> DMXSenderArray;
+
+	/** Buffer of the signals that are to be sent in the next frame */
+	TQueue<FDMXSignalSharedPtr> NewDMXSignalsToSend;
+
+	/** Map that holds the latest Signal per Universe */
+	FDMXThreadSafeUniverseToSignalMap LatestDMXSignals;
 
 	/** The Destination Address to send to, can be irrelevant, e.g. for art-net broadcast */
-	FString DestinationAddress;
+	TArray<FDMXOutputPortDestinationAddress> DestinationAddresses;
 
 	/** Helper to determine how dmx should be communicated (loopback, send) */
 	FDMXOutputPortCommunicationDeterminator CommunicationDeterminator;
 
 	/** Priority on which packets are being sent */
-	int32 Priority;
-
-	/** Map of latest Singals per Universe */
-	TMap<int32, FDMXSignalSharedPtr> ExternUniverseToLatestSignalMap;
+	int32 Priority = 0;
 
 	/** Map of raw Inputs */
 	TSet<TSharedRef<FDMXRawListener>> RawListeners;
@@ -160,5 +218,11 @@ private:
 	FGuid PortGuid;
 
 	/** Delay to apply on packets being sent */
-	float Delay;
+	double DelaySeconds = 0.0;
+
+	/** Holds the thread object. */
+	FRunnableThread* Thread = nullptr;
+
+	/** Flag indicating that the thread is stopping. */
+	TAtomic<bool> bStopping;
 };
diff --git a/Engine/Plugins/VirtualProduction/DMX/DMXProtocol/Source/DMXProtocol/Public/IO/DMXOutputPortConfig.h b/Engine/Plugins/VirtualProduction/DMX/DMXProtocol/Source/DMXProtocol/Public/IO/DMXOutputPortConfig.h
--- a/Engine/Plugins/VirtualProduction/DMX/DMXProtocol/Source/DMXProtocol/Public/IO/DMXOutputPortConfig.h
+++ b/Engine/Plugins/VirtualProduction/DMX/DMXProtocol/Source/DMXProtocol/Public/IO/DMXOutputPortConfig.h
@@ -16,6 +16,27 @@ class FDMXPort;
 struct FGuid;
 
 
+/** The IP address outbound DMX is sent to */
+USTRUCT(BlueprintType)
+struct DMXPROTOCOL_API FDMXOutputPortDestinationAddress
+{
+	GENERATED_BODY()
+
+	FDMXOutputPortDestinationAddress() = default;
+
+	FDMXOutputPortDestinationAddress(const FString& DestinationAddress)
+		: DestinationAddressString(DestinationAddress)
+	{}
+
+	/** The IP address outbound DMX is sent to */
+	UPROPERTY(Config, BlueprintReadOnly, EditDefaultsOnly, Category = "Port Config", Meta = (DisplayName = "Destination IP Address"))
+	FString DestinationAddressString;
+
+	FORCEINLINE bool operator==(const FDMXOutputPortDestinationAddress& Other) const { return DestinationAddressString == Other.DestinationAddressString; }
+	FORCEINLINE bool operator!=(const FDMXOutputPortDestinationAddress& Other) const { return !operator==(Other); }
+	FORCEINLINE friend uint32 GetTypeHash(const FDMXOutputPortDestinationAddress& Key) { return GetTypeHash(Key.DestinationAddressString); }
+};
+
 /** Data to create a new output port config with related constructor */
 struct DMXPROTOCOL_API FDMXOutputPortConfigParams
 {
@@ -26,13 +47,17 @@ struct DMXPROTOCOL_API FDMXOutputPortConfigParams
 	FName ProtocolName;
 	EDMXCommunicationType CommunicationType;
 	FString DeviceAddress;
-	FString DestinationAddress;
+	TArray<FDMXOutputPortDestinationAddress> DestinationAddresses;
 	bool bLoopbackToEngine;
 	int32 LocalUniverseStart;
 	int32 NumUniverses;
 	int32 ExternUniverseStart;
 	int32 Priority;
-	float Delay;
+	double DelaySeconds;
+	FFrameRate DelayFrameRate;
+
+	// DEPRECATED 5.0. Instead please use the DestinationAddresses array.
+	FString DestinationAddress_DEPRECATED;
 };
 
 /** 
@@ -62,22 +87,30 @@ public:
 	FORCEINLINE const FName& GetProtocolName() const { return ProtocolName; }
 	FORCEINLINE EDMXCommunicationType GetCommunicationType() const { return CommunicationType; }
 	FString GetDeviceAddress() const;
-	FORCEINLINE const FString& GetDestinationAddress() const { return DestinationAddress; }
+	FORCEINLINE const TArray<FDMXOutputPortDestinationAddress>& GetDestinationAddresses() const { return DestinationAddresses; }
 	FORCEINLINE bool NeedsLoopbackToEngine() const { return bLoopbackToEngine; }
 	FORCEINLINE int32 GetLocalUniverseStart() const { return LocalUniverseStart; }
 	FORCEINLINE int32 GetNumUniverses() const { return NumUniverses; }
 	FORCEINLINE int32 GetExternUniverseStart() const { return ExternUniverseStart; }
 	FORCEINLINE int32 GetPriority() const { return Priority; }
 	FORCEINLINE const FGuid& GetPortGuid() const { return PortGuid; }
-	FORCEINLINE int32 GetDelay() const { return Delay; }
+	int32 GetDelaySeconds() const;
+
+	UE_DEPRECATED(5.0, "Output Ports now support many destination addresses. Use FDMXOutputPortConfig::GetDestinationAddresses instead. ")
+	FORCEINLINE const FString& GetDestinationAddress() const { return DestinationAddress; }
 
 #if WITH_EDITOR
 	static FName GetProtocolNamePropertyNameChecked() { return GET_MEMBER_NAME_CHECKED(FDMXOutputPortConfig, ProtocolName); }
 	static FName GetCommunicationTypePropertyNameChecked() { return GET_MEMBER_NAME_CHECKED(FDMXOutputPortConfig, CommunicationType); }
 	static FName GetDeviceAddressPropertyNameChecked() { return GET_MEMBER_NAME_CHECKED(FDMXOutputPortConfig, DeviceAddress); }
-	static FName GetDestinationAddressPropertyNameChecked() { return GET_MEMBER_NAME_CHECKED(FDMXOutputPortConfig, DestinationAddress); }
+	static FName GetDestinationAddressesPropertyNameChecked() { return GET_MEMBER_NAME_CHECKED(FDMXOutputPortConfig, DestinationAddresses); }
 	static FName GetPriorityPropertyNameChecked() { return GET_MEMBER_NAME_CHECKED(FDMXOutputPortConfig, Priority); }
+	static FName GetDelayPropertyNameChecked() { return GET_MEMBER_NAME_CHECKED(FDMXOutputPortConfig, Delay); }
+	static FName GetDelayFrameRatePropertyNameChecked() { return GET_MEMBER_NAME_CHECKED(FDMXOutputPortConfig, DelayFrameRate); }
 	static FName GetPortGuidPropertyNameChecked() { return GET_MEMBER_NAME_CHECKED(FDMXOutputPortConfig, PortGuid); }
+
+	UE_DEPRECATED(5.0, "Output Ports now support many destination addresses. Use FDMXOutputPortConfig::GetDestinationAddressesPropertyNameChecked instead. ")
+	static FName GetDestinationAddressPropertyNameChecked() { return GET_MEMBER_NAME_CHECKED(FDMXOutputPortConfig, DestinationAddress); }
 #endif // WITH_EDITOR
 
 protected:
@@ -98,10 +131,14 @@ protected:
 	FString DeviceAddress = TEXT("127.0.0.1");
 
 	/** For Unicast, the IP address outbound DMX is sent to */
-	UPROPERTY(Config, BlueprintReadOnly, EditDefaultsOnly, Category = "Port Config", Meta = (DisplayName = "Destination IP Address"))
+	UPROPERTY(meta = (DeprecatedProperty, DeprecationMessage = "DestinationAddress is deprecated. Please use DestinationAddresses instead."))
 	FString DestinationAddress = TEXT("None");
 
-	/** If true, the signals of output to this port is input into to the engine. It will still show only under output ports and is not visible in Monitors as Input. */
+	/** For Unicast, the IP addresses outbound DMX is sent to */
+	UPROPERTY(Config, BlueprintReadOnly, EditDefaultsOnly, Category = "Port Config", Meta = (DisplayName = "Destination IP Address"))
+	TArray<FDMXOutputPortDestinationAddress> DestinationAddresses;
+
+	/** If true, the signals output from this port are input into to the engine. Note, signals input into the engine this way will not be visible in Monitors when monitoring Inputs. */
 	UPROPERTY(Config, BlueprintReadOnly, EditDefaultsOnly, Category = "Port Config", Meta = (DisplayName = "Input into Engine"))
 	bool bLoopbackToEngine = true;
 
@@ -124,9 +161,13 @@ protected:
 	UPROPERTY(Config, BlueprintReadOnly, EditDefaultsOnly, Category = "Port Config")
 	int32 Priority = 100;
 
-	/** Delay the sending of packets */
-	UPROPERTY(Config, BlueprintReadOnly, EditDefaultsOnly, Category = "Port Config", meta = (ClampMin = 0, ClampMax = 60))
-	float Delay = 0;
+	/** The amout by which sending of packets is delayed */
+	UPROPERTY(Config, BlueprintReadOnly, EditDefaultsOnly, Category = "Port Config")
+	double Delay = 0.0;
+
+	/** Framerate of the delay */
+	UPROPERTY(Config, BlueprintReadOnly, EditDefaultsOnly, Category = "Port Config")
+	FFrameRate DelayFrameRate;
 
 protected:
 	/** Generates a unique port name (unique for those stored in project settings) */
diff --git a/Engine/Plugins/VirtualProduction/DMX/DMXProtocol/Source/DMXProtocol/Public/Interfaces/IDMXProtocol.h b/Engine/Plugins/VirtualProduction/DMX/DMXProtocol/Source/DMXProtocol/Public/Interfaces/IDMXProtocol.h
--- a/Engine/Plugins/VirtualProduction/DMX/DMXProtocol/Source/DMXProtocol/Public/Interfaces/IDMXProtocol.h
+++ b/Engine/Plugins/VirtualProduction/DMX/DMXProtocol/Source/DMXProtocol/Public/Interfaces/IDMXProtocol.h
@@ -138,9 +138,9 @@ public:
 	 * This should be implemented so each port can only register once (see Art-Net for an example).
 	 *
 	 * @param	OutputPort			The Output port that needs to be registered
-	 * @return						The DMX sender, or nullptr if the port couldn't be registered.
+	 * @return						The DMX senders that can be used while the output port is registered.
 	 */
-	virtual TSharedPtr<IDMXSender> RegisterOutputPort(const FDMXOutputPortSharedRef& OutputPort) = 0;
+	virtual TArray<TSharedPtr<IDMXSender>> RegisterOutputPort(const FDMXOutputPortSharedRef& OutputPort) = 0;
 
 	/**
 	 * Called to unregister a DMXOutputPort with the protocol.
diff --git a/Engine/Plugins/VirtualProduction/DMX/DMXProtocol/Source/DMXProtocol/Public/Interfaces/IDMXSender.h b/Engine/Plugins/VirtualProduction/DMX/DMXProtocol/Source/DMXProtocol/Public/Interfaces/IDMXSender.h
--- a/Engine/Plugins/VirtualProduction/DMX/DMXProtocol/Source/DMXProtocol/Public/Interfaces/IDMXSender.h
+++ b/Engine/Plugins/VirtualProduction/DMX/DMXProtocol/Source/DMXProtocol/Public/Interfaces/IDMXSender.h
@@ -16,9 +16,10 @@ public:
 	virtual ~IDMXSender()
 	{}
 
-	/** Sends the DMX signal */
+	/** Sends the DMX signal. Note, this is called from a worker thread, not the game thread. */
 	virtual void SendDMXSignal(const FDMXSignalSharedRef& DMXSignal) = 0;
 
-	/** Clears teh buf*/
-	virtual void ClearBuffer() = 0;
+	/** DEPRECATED 5.0 */
+	UE_DEPRECATED(5.0, "Deprecated without replacement. Classes that implement IDMXSender do no longer need to implement this method.")
+	virtual void ClearBuffer() {};
 };
diff --git a/Engine/Plugins/VirtualProduction/DMX/DMXProtocol/Source/DMXProtocolArtNet/Private/DMXProtocolArtNet.cpp b/Engine/Plugins/VirtualProduction/DMX/DMXProtocol/Source/DMXProtocolArtNet/Private/DMXProtocolArtNet.cpp
--- a/Engine/Plugins/VirtualProduction/DMX/DMXProtocol/Source/DMXProtocolArtNet/Private/DMXProtocolArtNet.cpp
+++ b/Engine/Plugins/VirtualProduction/DMX/DMXProtocol/Source/DMXProtocolArtNet/Private/DMXProtocolArtNet.cpp
@@ -146,66 +146,78 @@ void FDMXProtocolArtNet::UnregisterInputPort(const FDMXInputPortSharedRef& Input
 	}
 }
 
-TSharedPtr<IDMXSender> FDMXProtocolArtNet::RegisterOutputPort(const FDMXOutputPortSharedRef& OutputPort)
+PRAGMA_DISABLE_OPTIMIZATION
+TArray<TSharedPtr<IDMXSender>> FDMXProtocolArtNet::RegisterOutputPort(const FDMXOutputPortSharedRef& OutputPort)
 {
 	check(!OutputPort->IsRegistered());
 	check(!CachedOutputPorts.Contains(OutputPort));
 
 	const FString& NetworkInterfaceAddress = OutputPort->GetDeviceAddress();
-	EDMXCommunicationType CommunicationType = OutputPort->GetCommunicationType();
+	const EDMXCommunicationType CommunicationType = OutputPort->GetCommunicationType();
 
 	// Try to use an existing receiver or create a new one
-	TSharedPtr<FDMXProtocolArtNetSender> Sender;
-	if (!Sender.IsValid())
+	TArray<TSharedPtr<IDMXSender>> NewSenders;
+	if (CommunicationType == EDMXCommunicationType::Broadcast)
 	{
-		if (CommunicationType == EDMXCommunicationType::Broadcast)
+		TSharedPtr<FDMXProtocolArtNetSender> NewSender = FindExistingBroadcastSender(NetworkInterfaceAddress);	
+		if (!NewSender.IsValid())
 		{
-			Sender = FindExistingBroadcastSender(NetworkInterfaceAddress);
-			
-			if (!Sender.IsValid())
-			{
-				Sender = FDMXProtocolArtNetSender::TryCreateBroadcastSender(SharedThis(this), NetworkInterfaceAddress);
-			}
+			NewSender = FDMXProtocolArtNetSender::TryCreateBroadcastSender(SharedThis(this), NetworkInterfaceAddress);
 		}
-		else if (CommunicationType == EDMXCommunicationType::Unicast)
-		{
-			const FString& UnicastAddress = OutputPort->GetDestinationAddress();
 
-			Sender = FindExistingUnicastSender(NetworkInterfaceAddress, UnicastAddress);
+		if (NewSender.IsValid())
+		{
+			NewSender->AssignOutputPort(OutputPort);
+			NewSenders.Add(NewSender);
+		}
+	}
+	else if (CommunicationType == EDMXCommunicationType::Unicast)
+	{
+		for (const FString& UnicastAddress : OutputPort->GetDestinationAddresses())
+		{
+			TSharedPtr<FDMXProtocolArtNetSender> NewSender = FindExistingUnicastSender(NetworkInterfaceAddress, UnicastAddress);
+			if (!NewSender.IsValid())
+			{
+				NewSender = FDMXProtocolArtNetSender::TryCreateUnicastSender(SharedThis(this), NetworkInterfaceAddress, UnicastAddress);
+			}
 
-			if (!Sender.IsValid())
+			if (NewSender.IsValid())
 			{
-				Sender = FDMXProtocolArtNetSender::TryCreateUnicastSender(SharedThis(this), NetworkInterfaceAddress, UnicastAddress);
+				NewSender->AssignOutputPort(OutputPort);
+				NewSenders.Add(NewSender);
 			}
 		}
-		else
-		{
-			UE_LOG(LogDMXProtocol, Error, TEXT("Cannot create DMX Protocol Art-Net Sender. The communication type specified is not supported."));
-		}
+	}
+	else
+	{
+		UE_LOG(LogDMXProtocol, Error, TEXT("Cannot create DMX Protocol Art-Net Sender. The communication type specified is not supported."));
 	}
 
-	if (!Sender.IsValid())
+	if (NewSenders.Num() > 0)
+	{
+		CachedOutputPorts.Add(OutputPort);
+	}
+	else
 	{
 		UE_LOG(LogDMXProtocol, Warning, TEXT("Could not create Art-Net sender for output port %s"), *OutputPort->GetPortName());
-
-		return nullptr;
 	}
 
-	Senders.Add(Sender);
-
-	Sender->AssignOutputPort(OutputPort);
-	CachedOutputPorts.Add(OutputPort);
+	for (const TSharedPtr<IDMXSender>& Sender : NewSenders)
+	{
+		Senders.Add(StaticCastSharedPtr<FDMXProtocolArtNetSender>(Sender));
+	}
 
-	return Sender;
+	return NewSenders;
 }
+PRAGMA_ENABLE_OPTIMIZATION
 
 void FDMXProtocolArtNet::UnregisterOutputPort(const FDMXOutputPortSharedRef& OutputPort)
 {
 	check(CachedOutputPorts.Contains(OutputPort));
 	CachedOutputPorts.Remove(OutputPort);
 
-	TSharedPtr<FDMXProtocolArtNetSender> UnusedSender;
-	for (const TSharedPtr<FDMXProtocolArtNetSender>& Sender : Senders)
+	TArray<TSharedPtr<FDMXProtocolArtNetSender>> UnusedSenders;
+	for (const TSharedPtr<FDMXProtocolArtNetSender>& Sender : TSet<TSharedPtr<FDMXProtocolArtNetSender>>(Senders))
 	{
 		if (Sender->ContainsOutputPort(OutputPort))
 		{
@@ -213,16 +225,10 @@ void FDMXProtocolArtNet::UnregisterOutputPort(const FDMXOutputPortSharedRef& Out
 
 			if (Sender->GetNumAssignedOutputPorts() == 0)
 			{
-				UnusedSender = Sender;
+				Senders.Remove(Sender);
 			}
-			break;
 		}
 	}
-
-	if (UnusedSender.IsValid())
-	{
-		Senders.Remove(UnusedSender);
-	}
 }
 
 bool FDMXProtocolArtNet::IsCausingLoopback(EDMXCommunicationType InCommunicationType)
diff --git a/Engine/Plugins/VirtualProduction/DMX/DMXProtocol/Source/DMXProtocolArtNet/Private/DMXProtocolArtNet.h b/Engine/Plugins/VirtualProduction/DMX/DMXProtocol/Source/DMXProtocolArtNet/Private/DMXProtocolArtNet.h
--- a/Engine/Plugins/VirtualProduction/DMX/DMXProtocol/Source/DMXProtocolArtNet/Private/DMXProtocolArtNet.h
+++ b/Engine/Plugins/VirtualProduction/DMX/DMXProtocol/Source/DMXProtocolArtNet/Private/DMXProtocolArtNet.h
@@ -38,7 +38,7 @@ public:
 	virtual bool SupportsPrioritySettings() const override;
 	virtual bool RegisterInputPort(const FDMXInputPortSharedRef& InputPort) override;
 	virtual void UnregisterInputPort(const FDMXInputPortSharedRef& InputPort) override;
-	virtual TSharedPtr<IDMXSender> RegisterOutputPort(const FDMXOutputPortSharedRef& OutputPort) override;
+	virtual TArray<TSharedPtr<IDMXSender>> RegisterOutputPort(const FDMXOutputPortSharedRef& OutputPort) override;
 	virtual void UnregisterOutputPort(const FDMXOutputPortSharedRef& OutputPort) override;
 	virtual bool IsCausingLoopback(EDMXCommunicationType InCommunicationType) override;
 	// ~End IDMXProtocol implementation
diff --git a/Engine/Plugins/VirtualProduction/DMX/DMXProtocol/Source/DMXProtocolArtNet/Private/DMXProtocolArtNetReceiver.cpp b/Engine/Plugins/VirtualProduction/DMX/DMXProtocol/Source/DMXProtocolArtNet/Private/DMXProtocolArtNetReceiver.cpp
--- a/Engine/Plugins/VirtualProduction/DMX/DMXProtocol/Source/DMXProtocolArtNet/Private/DMXProtocolArtNetReceiver.cpp
+++ b/Engine/Plugins/VirtualProduction/DMX/DMXProtocol/Source/DMXProtocolArtNet/Private/DMXProtocolArtNetReceiver.cpp
@@ -219,7 +219,7 @@ void FDMXProtocolArtNetReceiver::DistributeReceivedData(const TSharedRef<FArrayR
 uint16 FDMXProtocolArtNetReceiver::GetPacketOpCode(const TSharedRef<FArrayReader>& Buffer) const
 {
 	uint16 OpCode = 0x0000;
-	const uint32 MinCheck = ARTNET_STRING_SIZE + 2;
+	constexpr uint32 MinCheck = ARTNET_STRING_SIZE + 2;
 	if (Buffer->Num() > MinCheck)
 	{
 		// Get OpCode
@@ -236,7 +236,7 @@ uint16 FDMXProtocolArtNetReceiver::GetPacketOpCode(const TSharedRef<FArrayReader
 void FDMXProtocolArtNetReceiver::HandleDataPacket(const TSharedRef<FArrayReader>& PacketReader)
 {
 	uint16 UniverseID = 0x0000;
-	const uint32 MinCheck = ARTNET_UNIVERSE_ADDRESS + 2;
+	constexpr uint32 MinCheck = ARTNET_UNIVERSE_ADDRESS + 2;
 	if (PacketReader->Num() > MinCheck)
 	{
 		// Get OpCode
@@ -249,7 +249,8 @@ void FDMXProtocolArtNetReceiver::HandleDataPacket(const TSharedRef<FArrayReader>
 		FDMXProtocolArtNetDMXPacket ArtNetDMXPacket;
 		*PacketReader << ArtNetDMXPacket;
 
-		FDMXSignalSharedRef DMXSignal = MakeShared<FDMXSignal, ESPMode::ThreadSafe>(FPlatformTime::Seconds(), UniverseID, 0, TArray<uint8>(ArtNetDMXPacket.Data, DMX_UNIVERSE_SIZE), 0);
+		constexpr int32 Priority = 0;
+		FDMXSignalSharedRef DMXSignal = MakeShared<FDMXSignal, ESPMode::ThreadSafe>(FPlatformTime::Seconds(), UniverseID, Priority, TArray<uint8>(ArtNetDMXPacket.Data, DMX_UNIVERSE_SIZE));
 		for (const TSharedPtr<FDMXInputPort, ESPMode::ThreadSafe>& InputPort : AssignedInputPorts)
 		{				
 			InputPort->SingleProducerInputDMXSignal(DMXSignal);
diff --git a/Engine/Plugins/VirtualProduction/DMX/DMXProtocol/Source/DMXProtocolArtNet/Private/DMXProtocolArtNetReceiver.h b/Engine/Plugins/VirtualProduction/DMX/DMXProtocol/Source/DMXProtocolArtNet/Private/DMXProtocolArtNetReceiver.h
--- a/Engine/Plugins/VirtualProduction/DMX/DMXProtocol/Source/DMXProtocolArtNet/Private/DMXProtocolArtNetReceiver.h
+++ b/Engine/Plugins/VirtualProduction/DMX/DMXProtocol/Source/DMXProtocolArtNet/Private/DMXProtocolArtNetReceiver.h
@@ -107,7 +107,7 @@ private:
 	TSharedPtr<FDMXProtocolArtNet, ESPMode::ThreadSafe> Protocol;
 
 	/** The network socket. */
-	FSocket* Socket;
+	FSocket* Socket = nullptr;
 
 	/** The endpoint internet addr */
 	TSharedPtr<FInternetAddr> EndpointInternetAddr;
@@ -116,7 +116,7 @@ private:
 	TAtomic<bool> bStopping;
 
 	/** The thread object. */
-	FRunnableThread* Thread;
+	FRunnableThread* Thread = nullptr;
 
 	/** The receiver thread's name. */
 	FString ThreadName;
diff --git a/Engine/Plugins/VirtualProduction/DMX/DMXProtocol/Source/DMXProtocolArtNet/Private/DMXProtocolArtNetSender.cpp b/Engine/Plugins/VirtualProduction/DMX/DMXProtocol/Source/DMXProtocolArtNet/Private/DMXProtocolArtNetSender.cpp
--- a/Engine/Plugins/VirtualProduction/DMX/DMXProtocol/Source/DMXProtocolArtNet/Private/DMXProtocolArtNetSender.cpp
+++ b/Engine/Plugins/VirtualProduction/DMX/DMXProtocol/Source/DMXProtocolArtNet/Private/DMXProtocolArtNetSender.cpp
@@ -14,7 +14,6 @@
 #include "SocketSubsystem.h"
 #include "Common/UdpSocketBuilder.h"
 #include "Common/UdpSocketSender.h"
-#include "HAL/RunnableThread.h"
 #include "Serialization/ArrayReader.h"
 #include "Serialization/ArrayWriter.h"
 #include "UObject/Class.h"
@@ -27,14 +26,9 @@ FDMXProtocolArtNetSender::FDMXProtocolArtNetSender(const TSharedPtr<FDMXProtocol
 	, Socket(&InSocket)
 	, NetworkInterfaceInternetAddr(InNetworkInterfaceInternetAddr)
 	, DestinationInternetAddr(InDestinationInternetAddr)
-	, bStopping(false)
-	, Thread(nullptr)
 {
 	check(DestinationInternetAddr.IsValid());
 
-	FString SenderThreadName = FString(TEXT("ArtNetSender_")) + InDestinationInternetAddr->ToString(false);
-	Thread = FRunnableThread::Create(this, *SenderThreadName, 0U, TPri_TimeCritical, FPlatformAffinity::GetPoolThreadMask());
-
 	UE_LOG(LogDMXProtocol, VeryVerbose, TEXT("Created Art-Net Sender at %s sending to %s"), *NetworkInterfaceInternetAddr->ToString(false), *DestinationInternetAddr->ToString(false));
 }
 
@@ -46,12 +40,6 @@ FDMXProtocolArtNetSender::~FDMXProtocolArtNetSender()
 		SocketSubsystem->DestroySocket(Socket);
 	}
 
-	if (Thread != nullptr)
-	{
-		Thread->Kill(true);
-		delete Thread;
-	}
-
 	UE_LOG(LogDMXProtocol, VeryVerbose, TEXT("Destroyed Art-Net Sender at %s sending to %s"), *NetworkInterfaceInternetAddr->ToString(false), *DestinationInternetAddr->ToString(false));
 }
 
@@ -166,15 +154,48 @@ bool FDMXProtocolArtNetSender::IsCausingLoopback() const
 
 void FDMXProtocolArtNetSender::SendDMXSignal(const FDMXSignalSharedRef& DMXSignal)
 {
-	Buffer.Enqueue(DMXSignal);
-}
+	FDMXProtocolArtNetDMXPacket ArtNetDMXPacket;
+	FMemory::Memcpy(ArtNetDMXPacket.Data, DMXSignal->ChannelData.GetData(), ARTNET_DMX_LENGTH);
 
-void FDMXProtocolArtNetSender::ClearBuffer()
-{
-	FScopeLock Lock(&LatestSignalLock);
+	uint16 UniverseID = static_cast<uint16>(DMXSignal->ExternUniverseID);
 
-	Buffer.Empty();
-	UniverseToLatestSignalMap.Reset();
+	//Set Packet Data
+	ArtNetDMXPacket.Physical = 0; // As per Standard: For information only. We always specify port 0.
+	ArtNetDMXPacket.Universe = UniverseID;
+	ArtNetDMXPacket.Sequence = 0x00; // As per Standard: The Sequence field is set to 0x00 to disable this feature.
+
+	TSharedPtr<FBufferArchive> BufferArchive = ArtNetDMXPacket.Pack(ARTNET_DMX_LENGTH);
+
+	int32 SendDataSize = BufferArchive->Num();
+	int32 BytesSent = -1;
+
+	// Try to send, log errors but avoid spaming the Log
+	static bool bErrorEverLogged = false;
+	if (Socket->SendTo(BufferArchive->GetData(), BufferArchive->Num(), BytesSent, *DestinationInternetAddr))
+	{
+		INC_DWORD_STAT(STAT_ArtNetPackagesSent);
+	}
+	else
+	{
+		if (!bErrorEverLogged)
+		{
+			ISocketSubsystem* SocketSubsystem = ISocketSubsystem::Get(PLATFORM_SOCKETSUBSYSTEM);
+			TEnumAsByte<ESocketErrors> RecvFromError = SocketSubsystem->GetLastErrorCode();
+
+			UE_LOG(LogDMXProtocol, Error, TEXT("Failed send DMX to %s with Error Code %d"), *DestinationInternetAddr->ToString(false), RecvFromError);
+
+			bErrorEverLogged = true;
+		}
+	}
+
+	if (BytesSent != SendDataSize)
+	{
+		if (!bErrorEverLogged)
+		{
+			UE_LOG(LogDMXProtocol, Warning, TEXT("Incomplete DMX Packet sent to %s"), *DestinationInternetAddr->ToString(false));
+			bErrorEverLogged = true;
+		}
+	}
 }
 
 TSharedPtr<FInternetAddr> FDMXProtocolArtNetSender::CreateInternetAddr(const FString& IPAddress, int32 Port)
@@ -206,164 +227,3 @@ TSharedRef<FInternetAddr> FDMXProtocolArtNetSender::CreateBroadcastInternetAddr(
 
 	return InternetAddr;
 }
-
-bool FDMXProtocolArtNetSender::Init()
-{
-	return true;
-}
-
-uint32 FDMXProtocolArtNetSender::Run()
-{
-	const UDMXProtocolSettings* DMXSettings = GetDefault<UDMXProtocolSettings>();
-	check(DMXSettings);
-	
-	// Fixed rate delta time
-	const double SendDeltaTime = 1.f / DMXSettings->SendingRefreshRate;
-
-	while (!bStopping)
-	{
-		const double StartTime = FPlatformTime::Seconds();
-
-		Update();
-
-		const double EndTime = FPlatformTime::Seconds();
-		const double WaitTime = SendDeltaTime - (EndTime - StartTime);
-
-		if (WaitTime > 0.f)
-		{
-			// Sleep by the amount which is set in refresh rate
-			FPlatformProcess::SleepNoStats(WaitTime);
-		}
-
-		// In the unlikely case we took to long to send, we instantly continue, but do not take 
-		// further measures to compensate - We would have to run faster than DMX send rate to catch up.
-	}
-
-	return 0;
-}
-
-void FDMXProtocolArtNetSender::Stop()
-{
-	bStopping = true;
-}
-
-void FDMXProtocolArtNetSender::Exit()
-{
-}
-
-void FDMXProtocolArtNetSender::Tick()
-{
-	Update();
-}
-
-FSingleThreadRunnable* FDMXProtocolArtNetSender::GetSingleThreadInterface()
-{
-	return this;
-}
-
-void FDMXProtocolArtNetSender::Update()
-{
-	UniverseToLatestSignalMap.Reset();
-
-	// process delayed signals
-	const double Now = FPlatformTime::Seconds();
-
-	for (;;)
-	{
-		// we can safely assumes that older messages are on tail of the queue
-		FDMXSignal OldestDMXSignal;
-		if (DelayedBuffer.Peek(OldestDMXSignal))
-		{
-			if ((OldestDMXSignal.Timestamp + static_cast<double>(OldestDMXSignal.Delay)) <= Now)
-			{
-				if (UniverseToLatestSignalMap.Contains(OldestDMXSignal.ExternUniverseID))
-				{
-					UniverseToLatestSignalMap[OldestDMXSignal.ExternUniverseID] = MakeShared<FDMXSignal>(OldestDMXSignal);
-				}
-				else
-				{
-					UniverseToLatestSignalMap.Add(OldestDMXSignal.ExternUniverseID, MakeShared<FDMXSignal>(OldestDMXSignal));
-				}
-				DelayedBuffer.Pop();
-				continue;
-			}
-		}
-		break;
-	}
-
-	{
-		FScopeLock Lock(&LatestSignalLock);
-
-		// Keep latest signal per universe
-		FDMXSignalSharedPtr DequeuedDMXSignal;
-		while (Buffer.Dequeue(DequeuedDMXSignal))
-		{
-			if (Protocol->IsValidUniverseID(DequeuedDMXSignal->ExternUniverseID))
-			{
-				if (DequeuedDMXSignal->Delay > 0)
-				{
-					// we make a copy here of the signal for not losing track
-					DelayedBuffer.Enqueue(*DequeuedDMXSignal);
-					continue;
-				}
-				if (UniverseToLatestSignalMap.Contains(DequeuedDMXSignal->ExternUniverseID))
-				{
-					UniverseToLatestSignalMap[DequeuedDMXSignal->ExternUniverseID] = DequeuedDMXSignal.ToSharedRef();
-				}
-				else
-				{
-					UniverseToLatestSignalMap.Add(DequeuedDMXSignal->ExternUniverseID, DequeuedDMXSignal.ToSharedRef());
-				}
-			}
-		}
-	}
-
-	// Create a packet for each universe and send it
-	for (const TTuple<int32, FDMXSignalSharedRef>& UniverseToSignalKvp : UniverseToLatestSignalMap)
-	{
-		const FDMXSignalSharedRef& DMXSignal = UniverseToSignalKvp.Value;
-
-		FDMXProtocolArtNetDMXPacket ArtNetDMXPacket;
-		FMemory::Memcpy(ArtNetDMXPacket.Data, DMXSignal->ChannelData.GetData(), ARTNET_DMX_LENGTH);
-
-		uint16 UniverseID = static_cast<uint16>(DMXSignal->ExternUniverseID);
-
-		//Set Packet Data
-		ArtNetDMXPacket.Physical = 0; // As per Standard: For information only. We always specify port 0.
-		ArtNetDMXPacket.Universe = UniverseID;
-		ArtNetDMXPacket.Sequence = 0x00; // As per Standard: The Sequence field is set to 0x00 to disable this feature.
-
-		TSharedPtr<FBufferArchive> BufferArchive = ArtNetDMXPacket.Pack(ARTNET_DMX_LENGTH);
-
-		int32 SendDataSize = BufferArchive->Num();
-		int32 BytesSent = -1;
-
-		// Try to send, log errors but avoid spaming the Log
-		static bool bErrorEverLogged = false;
-		if (Socket->SendTo(BufferArchive->GetData(), BufferArchive->Num(), BytesSent, *DestinationInternetAddr))
-		{
-			INC_DWORD_STAT(STAT_ArtNetPackagesSent);
-		}
-		else
-		{
-			if(!bErrorEverLogged)
-			{
-				ISocketSubsystem* SocketSubsystem = ISocketSubsystem::Get(PLATFORM_SOCKETSUBSYSTEM);
-				TEnumAsByte<ESocketErrors> RecvFromError = SocketSubsystem->GetLastErrorCode();
-
-				UE_LOG(LogDMXProtocol, Error, TEXT("Failed send DMX to %s with Error Code %d"), *DestinationInternetAddr->ToString(false), RecvFromError);
-
-				bErrorEverLogged = true;
-			}
-		}
-		
-		if (BytesSent != SendDataSize)
-		{
-			if (!bErrorEverLogged)
-			{
-				UE_LOG(LogDMXProtocol, Warning, TEXT("Incomplete DMX Packet sent to %s"), *DestinationInternetAddr->ToString(false));
-				bErrorEverLogged = true;
-			}
-		}
-	}
-}
diff --git a/Engine/Plugins/VirtualProduction/DMX/DMXProtocol/Source/DMXProtocolArtNet/Private/DMXProtocolArtNetSender.h b/Engine/Plugins/VirtualProduction/DMX/DMXProtocol/Source/DMXProtocolArtNet/Private/DMXProtocolArtNetSender.h
--- a/Engine/Plugins/VirtualProduction/DMX/DMXProtocol/Source/DMXProtocolArtNet/Private/DMXProtocolArtNetSender.h
+++ b/Engine/Plugins/VirtualProduction/DMX/DMXProtocol/Source/DMXProtocolArtNet/Private/DMXProtocolArtNetSender.h
@@ -7,13 +7,7 @@
 #include "Interfaces/IDMXSender.h"
 
 #include "CoreMinimal.h"
-#include "Containers/Queue.h"
-#include "HAL/CriticalSection.h"
-#include "HAL/Runnable.h"
-#include "Misc/ScopeLock.h" 
-#include "Misc/SingleThreadRunnable.h"
 #include "Serialization/ArrayWriter.h"
-#include "Templates/Atomic.h"
 
 enum class EDMXCommunicationType : uint8;
 struct FDMXProtocolArtNetRDM;
@@ -21,15 +15,12 @@ class FDMXOutputPort;
 
 class FDMXProtocolArtNet;
 class FInternetAddr;
-class FRunnableThread;
 class FSocket;
 class ISocketSubsystem;
 
 
 class DMXPROTOCOLARTNET_API FDMXProtocolArtNetSender
-	: public FRunnable
-	, public FSingleThreadRunnable 
-	, public IDMXSender
+	: public IDMXSender
 {
 protected:
 	/** Constructor. Hidden on purpose, use Create instead. */
@@ -62,7 +53,6 @@ public:
 
 	// ~Begin IDMXSender Interface
 	virtual void SendDMXSignal(const FDMXSignalSharedRef& DMXSignal) override;
-	virtual void ClearBuffer() override;
 	// ~End IDMXSender Interface
 
 public:
@@ -97,35 +87,11 @@ private:
 	/** The Output ports the receiver uses */
 	TSet<TSharedPtr<FDMXOutputPort, ESPMode::ThreadSafe>> AssignedOutputPorts;
 
-protected:
-	//~ Begin FRunnable implementation
-	virtual bool Init() override;
-	virtual uint32 Run() override;
-	virtual void Stop() override;
-	virtual void Exit() override;
-	//~ End FRunnable implementation
-
-	//~ Begin FSingleThreadRunnable implementation
-	virtual void Tick() override;
-	virtual class FSingleThreadRunnable* GetSingleThreadInterface() override;
-	//~ End FSingleThreadRunnable implementation
-
-protected:
-	/** Updates the thread, sending DMX */
-	void Update();
-
-private:
-	/** Buffer of dmx signals */
-	TQueue<FDMXSignalSharedPtr> Buffer;
-
-	/** Map of the latest signal per universe */
-	TMap<int32 /** Universe */, FDMXSignalSharedRef> UniverseToLatestSignalMap;
-
 	/** The Art-Net protocol instance */
 	TSharedPtr<FDMXProtocolArtNet, ESPMode::ThreadSafe> Protocol;
 
 	/** Holds the network socket used to sender packages. */
-	FSocket* Socket;
+	FSocket* Socket = nullptr;
 
 	/** The network interface internet addr */
 	TSharedPtr<FInternetAddr> NetworkInterfaceInternetAddr;
@@ -135,16 +101,4 @@ private:
 
 	/** Communication type used for the network traffic */
 	EDMXCommunicationType CommunicationType;
-
-	/** Lock used when the buffer is cleared */
-	FCriticalSection LatestSignalLock;
-
-	/** Flag indicating that the thread is stopping. */
-	TAtomic<bool> bStopping;
-
-	/** Holds the thread object. */
-	FRunnableThread* Thread;
-
-	/** Buffer of delayed dmx signals */
-	TQueue<FDMXSignal> DelayedBuffer;
 };
diff --git a/Engine/Plugins/VirtualProduction/DMX/DMXProtocol/Source/DMXProtocolEditor/DMXProtocolEditor.Build.cs b/Engine/Plugins/VirtualProduction/DMX/DMXProtocol/Source/DMXProtocolEditor/DMXProtocolEditor.Build.cs
--- a/Engine/Plugins/VirtualProduction/DMX/DMXProtocol/Source/DMXProtocolEditor/DMXProtocolEditor.Build.cs
+++ b/Engine/Plugins/VirtualProduction/DMX/DMXProtocol/Source/DMXProtocolEditor/DMXProtocolEditor.Build.cs
@@ -28,6 +28,7 @@ public class DMXProtocolEditor : ModuleRules
 				"InputCore",
 				"Slate",
 				"SlateCore",
+				"TimeManagement",
 				"UnrealEd"
 			}
 		);
diff --git a/Engine/Plugins/VirtualProduction/DMX/DMXProtocol/Source/DMXProtocolEditor/Private/DMXProtocolEditorModule.cpp b/Engine/Plugins/VirtualProduction/DMX/DMXProtocol/Source/DMXProtocolEditor/Private/DMXProtocolEditorModule.cpp
--- a/Engine/Plugins/VirtualProduction/DMX/DMXProtocol/Source/DMXProtocolEditor/Private/DMXProtocolEditorModule.cpp
+++ b/Engine/Plugins/VirtualProduction/DMX/DMXProtocol/Source/DMXProtocolEditor/Private/DMXProtocolEditorModule.cpp
@@ -7,6 +7,7 @@
 #include "IO/DMXOutputPortReference.h"
 #include "DetailsCustomizations/DMXInputPortConfigCustomization.h"
 #include "DetailsCustomizations/DMXOutputPortConfigCustomization.h"
+#include "DetailsCustomizations/DMXOutputPortDestinationAddressCustomization.h"
 #include "DetailsCustomizations/DMXInputPortReferenceCustomization.h"
 #include "DetailsCustomizations/DMXOutputPortReferenceCustomization.h"
 
@@ -42,6 +43,8 @@ void FDMXProtocolEditorModule::RegisterDetailsCustomizations()
 	PropertyModule.RegisterCustomPropertyTypeLayout(FDMXInputPortConfig::StaticStruct()->GetFName(), FOnGetPropertyTypeCustomizationInstance::CreateStatic(&FDMXInputPortConfigCustomization::MakeInstance));
 	PropertyModule.RegisterCustomPropertyTypeLayout(FDMXOutputPortConfig::StaticStruct()->GetFName(), FOnGetPropertyTypeCustomizationInstance::CreateStatic(&FDMXOutputPortConfigCustomization::MakeInstance));
 	
+	PropertyModule.RegisterCustomPropertyTypeLayout(FDMXOutputPortDestinationAddress::StaticStruct()->GetFName(), FOnGetPropertyTypeCustomizationInstance::CreateStatic(&FDMXOutputPortDestinationAddressCustomization::MakeInstance));
+
 	PropertyModule.RegisterCustomPropertyTypeLayout(FDMXInputPortReference::StaticStruct()->GetFName(), FOnGetPropertyTypeCustomizationInstance::CreateStatic(&FDMXInputPortReferenceCustomization::MakeInstance));
 	PropertyModule.RegisterCustomPropertyTypeLayout(FDMXOutputPortReference::StaticStruct()->GetFName(), FOnGetPropertyTypeCustomizationInstance::CreateStatic(&FDMXOutputPortReferenceCustomization::MakeInstance));
 
diff --git a/Engine/Plugins/VirtualProduction/DMX/DMXProtocol/Source/DMXProtocolEditor/Private/DetailsCustomizations/DMXInputPortConfigCustomization.cpp b/Engine/Plugins/VirtualProduction/DMX/DMXProtocol/Source/DMXProtocolEditor/Private/DetailsCustomizations/DMXInputPortConfigCustomization.cpp
--- a/Engine/Plugins/VirtualProduction/DMX/DMXProtocol/Source/DMXProtocolEditor/Private/DetailsCustomizations/DMXInputPortConfigCustomization.cpp
+++ b/Engine/Plugins/VirtualProduction/DMX/DMXProtocol/Source/DMXProtocolEditor/Private/DetailsCustomizations/DMXInputPortConfigCustomization.cpp
@@ -16,52 +16,271 @@
 #include "DetailWidgetRow.h"
 #include "IDetailChildrenBuilder.h"
 #include "IDetailPropertyRow.h" 
-#include "IPAddress.h" 
-#include "SocketSubsystem.h"
+#include "IPropertyUtilities.h"
+#include "ScopedTransaction.h"
 #include "Misc/Guid.h" 
 #include "Widgets/Text/STextBlock.h" 
 
 
-#define LOCTEXT_NAMESPACE "DMXPortConfigCustomization"
+#define LOCTEXT_NAMESPACE "DMXInputPortConfigCustomization"
 
 TSharedRef<IPropertyTypeCustomization> FDMXInputPortConfigCustomization::MakeInstance()
 {
 	return MakeShared<FDMXInputPortConfigCustomization>();
 }
 
-FName FDMXInputPortConfigCustomization::GetProtocolNamePropertyName() const
+void FDMXInputPortConfigCustomization::CustomizeHeader(TSharedRef<IPropertyHandle> StructPropertyHandle, class FDetailWidgetRow& HeaderRow, IPropertyTypeCustomizationUtils& StructCustomizationUtils)
 {
-	return FDMXInputPortConfig::GetProtocolNamePropertyNameChecked();
+	PropertyUtilities = StructCustomizationUtils.GetPropertyUtilities();
+
+	const bool bDisplayResetToDefault = false;
+	const FText DisplayNameOverride = FText::GetEmpty();
+	const FText DisplayToolTipOverride = FText::GetEmpty();
+
+	HeaderRow
+	.NameContent()
+	[
+		StructPropertyHandle->CreatePropertyNameWidget(DisplayNameOverride, DisplayToolTipOverride, bDisplayResetToDefault)
+	];
+}
+
+void FDMXInputPortConfigCustomization::CustomizeChildren(TSharedRef<IPropertyHandle> StructPropertyHandle, class IDetailChildrenBuilder& ChildBuilder, IPropertyTypeCustomizationUtils& StructCustomizationUtils)
+{
+	// Retrieve structure's child properties
+	uint32 NumChildren;
+	StructPropertyHandle->GetNumChildren(NumChildren);
+	TMap<FName, TSharedPtr< IPropertyHandle>> PropertyHandles;
+	for (uint32 ChildIndex = 0; ChildIndex < NumChildren; ++ChildIndex)
+	{
+		TSharedRef<IPropertyHandle> ChildHandle = StructPropertyHandle->GetChildHandle(ChildIndex).ToSharedRef();
+		const FName PropertyName = ChildHandle->GetProperty()->GetFName();
+
+		PropertyHandles.Add(PropertyName, ChildHandle);
+	}
+
+	// Cache customized properties
+	ProtocolNameHandle = PropertyHandles.FindChecked(FDMXInputPortConfig::GetProtocolNamePropertyNameChecked());
+	CommunicationTypeHandle = PropertyHandles.FindChecked(FDMXInputPortConfig::GetCommunicationTypePropertyNameChecked());
+	DeviceAddressHandle = PropertyHandles.FindChecked(FDMXInputPortConfig::GetDeviceAddressPropertyNameChecked());
+	PortGuidHandle = PropertyHandles.FindChecked(FDMXInputPortConfig::GetPortGuidPropertyNameChecked());
+
+	// Ports always need a valid Guid (cannot be blueprinted)
+	if (!GetPortGuid().IsValid())
+	{
+		ChildBuilder.AddCustomRow(LOCTEXT("InvalidPortConfigSearchString", "Invalid"))
+			.WholeRowContent()
+			[
+				SNew(STextBlock)
+				.Text(LOCTEXT("InvalidPortConfigText", "Invalid Port Guid. Cannot utilize this port."))
+			];
+
+		return;
+	}
+
+	// Add customized properties
+	for (auto Iter(PropertyHandles.CreateConstIterator()); Iter; ++Iter)
+	{
+		// Don't add the PortGuid property
+		if (Iter.Value() == PortGuidHandle)
+		{
+			continue;
+		}
+		
+		// Add the property
+		IDetailPropertyRow& PropertyRow = ChildBuilder.AddProperty(Iter.Value().ToSharedRef());
+
+		if (Iter.Value() == ProtocolNameHandle)
+		{
+			GenerateProtocolNameRow(PropertyRow);
+		}
+		else if (Iter.Value() == CommunicationTypeHandle)
+		{
+			GenerateCommunicationTypeRow(PropertyRow);
+		}
+		else if (Iter.Value() == DeviceAddressHandle)
+		{
+			GenerateIPAddressRow(PropertyRow);
+		}
+		else if (Iter.Key() == FDMXInputPortConfig::GetPriorityStrategyPropertyNameChecked())
+		{
+			TAttribute<EVisibility> PriorityStrategyVisibilityAttribute = TAttribute<EVisibility>::Create(TAttribute<EVisibility>::FGetter::CreateLambda([this]() 
+				{
+					const bool bPriorityStrategyIsVisible = GetProtocol().IsValid() && GetProtocol()->SupportsPrioritySettings();
+					return bPriorityStrategyIsVisible ? EVisibility::Visible : EVisibility::Collapsed;
+				}));
+
+			PropertyRow.Visibility(PriorityStrategyVisibilityAttribute);
+		}
+		else if (Iter.Key() == FDMXInputPortConfig::GetPriorityPropertyNameChecked())
+		{
+			TAttribute<EVisibility> PriorityVisibilityAttribute = TAttribute<EVisibility>::Create(TAttribute<EVisibility>::FGetter::CreateLambda([this]() 
+				{
+					const bool bPriorityIsVisible = GetProtocol().IsValid() && GetProtocol()->SupportsPrioritySettings();
+					return bPriorityIsVisible ? EVisibility::Visible : EVisibility::Collapsed;
+				}));
+
+			PropertyRow.Visibility(PriorityVisibilityAttribute);
+		}
+	}
+}
+
+void FDMXInputPortConfigCustomization::GenerateProtocolNameRow(IDetailPropertyRow& PropertyRow)
+{
+	// Customizate the Protocol Name property to draw a combo box with all protocol names
+
+	TSharedPtr<SWidget> NameWidget;
+	TSharedPtr<SWidget> ValueWidget;
+	FDetailWidgetRow Row;
+	PropertyRow.GetDefaultWidgets(NameWidget, ValueWidget, Row);
+	
+	const FName InitialSelection = [this]() -> FName
+	{
+		if (const IDMXProtocolPtr& Protocol = GetProtocol())
+		{
+			return Protocol->GetProtocolName();
+		}
+
+		return NAME_None;
+	}();
+
+	PropertyRow.CustomWidget()
+		.NameContent()
+		[
+			NameWidget.ToSharedRef()
+		]
+		.ValueContent()
+		[
+			SAssignNew(ProtocolNameComboBox, SDMXProtocolNameComboBox)
+			.InitiallySelectedProtocolName(InitialSelection)
+			.OnProtocolNameSelected(this, &FDMXInputPortConfigCustomization::OnProtocolNameSelected)
+		];
+
+}
+
+void FDMXInputPortConfigCustomization::GenerateCommunicationTypeRow(IDetailPropertyRow& PropertyRow)
+{
+	// Customizate the Communication Type property depending on the port's supported communication modes
+
+	TAttribute<EVisibility> CommunicationTypeVisibilityAttribute = TAttribute<EVisibility>::Create(TAttribute<EVisibility>::FGetter::CreateSP(this, &FDMXInputPortConfigCustomization::GetCommunicationTypeVisibility));
+	PropertyRow.Visibility(CommunicationTypeVisibilityAttribute);
+
+	TSharedPtr<SWidget> NameWidget;
+	TSharedPtr<SWidget> ValueWidget;
+	FDetailWidgetRow Row;
+	PropertyRow.GetDefaultWidgets(NameWidget, ValueWidget, Row);
+
+	PropertyRow.CustomWidget()
+		.Visibility(CommunicationTypeVisibilityAttribute)
+		.NameContent()
+		[
+			NameWidget.ToSharedRef()
+		]
+		.ValueContent()
+		[
+			SAssignNew(CommunicationTypeComboBox, SDMXCommunicationTypeComboBox)
+			.CommunicationTypes(GetSupportedCommunicationTypes())
+			.InitialCommunicationType(GetCommunicationType())
+			.OnCommunicationTypeSelected(this, &FDMXInputPortConfigCustomization::OnCommunicationTypeSelected)
+		];
 }
 
-FName FDMXInputPortConfigCustomization::GetCommunicationTypePropertyName() const
+void FDMXInputPortConfigCustomization::GenerateIPAddressRow(IDetailPropertyRow& PropertyRow)
+{
+	// Customizate the IPAddress property to show a combo box with available IP addresses
+	FString InitialValue = GetIPAddress();
+
+	TSharedPtr<SWidget> NameWidget;
+	TSharedPtr<SWidget> ValueWidget;
+	FDetailWidgetRow Row;
+	PropertyRow.GetDefaultWidgets(NameWidget, ValueWidget, Row);
+
+	PropertyRow.CustomWidget()
+		.NameContent()
+		[
+			NameWidget.ToSharedRef()
+		]
+		.ValueContent()
+		[
+			SAssignNew(IPAddressEditWidget, SDMXIPAddressEditWidget)
+			.InitialValue(InitialValue)
+			.bShowLocalNICComboBox(true)
+			.OnIPAddressSelected(this, &FDMXInputPortConfigCustomization::OnIPAddressSelected)
+		];
+}	
+
+EVisibility FDMXInputPortConfigCustomization::GetCommunicationTypeVisibility() const
 {
-	return FDMXInputPortConfig::GetCommunicationTypePropertyNameChecked();
+	if (CommunicationTypeComboBox.IsValid() && GetProtocol().IsValid())
+	{
+		return CommunicationTypeComboBox->GetVisibility();
+	}
+
+	return EVisibility::Collapsed;
 }
 
-FName FDMXInputPortConfigCustomization::GetDeviceAddressPropertyName() const
+void FDMXInputPortConfigCustomization::OnProtocolNameSelected()
 {
-	return FDMXInputPortConfig::GetDeviceAddressPropertyNameChecked();
+	FName ProtocolName = ProtocolNameComboBox->GetSelectedProtocolName();
+
+	const FScopedTransaction Transaction(LOCTEXT("ProtocolSelected", "DMX: Selected Protocol"));
+	
+	ProtocolNameHandle->NotifyPreChange();
+	ProtocolNameHandle->SetValue(ProtocolName);
+	ProtocolNameHandle->NotifyPostChange(EPropertyChangeType::ValueSet);
+
+	PropertyUtilities->ForceRefresh();
 }
 
-FName FDMXInputPortConfigCustomization::GetDestinationAddressPropertyName() const
+void FDMXInputPortConfigCustomization::OnCommunicationTypeSelected()
 {
-	return NAME_None;
+	EDMXCommunicationType SelectedCommunicationType = CommunicationTypeComboBox->GetSelectedCommunicationType();
+
+	const FScopedTransaction Transaction(LOCTEXT("CommunicationTypeSelected", "DMX: Selected Communication Type"));
+	
+	CommunicationTypeHandle->NotifyPreChange();
+	CommunicationTypeHandle->SetValue(static_cast<uint8>(SelectedCommunicationType));
+	CommunicationTypeHandle->NotifyPostChange(EPropertyChangeType::ValueSet);
+
+	PropertyUtilities->ForceRefresh();
 }
 
-FName FDMXInputPortConfigCustomization::GetPriorityStrategyPropertyName() const
+void FDMXInputPortConfigCustomization::OnIPAddressSelected()
 {
-	return FDMXInputPortConfig::GetPriorityStrategyPropertyNameChecked();
+	FString SelectedIP = IPAddressEditWidget->GetSelectedIPAddress();
+
+	const FScopedTransaction Transaction(LOCTEXT("CommunicationTypeSelected", "DMX: Selected IP Address"));
+	
+	DeviceAddressHandle->NotifyPreChange();
+	ensure(DeviceAddressHandle->SetValue(SelectedIP) == FPropertyAccess::Success);	
+	DeviceAddressHandle->NotifyPostChange(EPropertyChangeType::ValueSet);
 }
 
-FName FDMXInputPortConfigCustomization::GetPriorityPropertyName() const
+IDMXProtocolPtr FDMXInputPortConfigCustomization::GetProtocol() const
 {
-	return FDMXInputPortConfig::GetPriorityPropertyNameChecked();
+	FName ProtocolName;
+	ProtocolNameHandle->GetValue(ProtocolName);
+
+	IDMXProtocolPtr Protocol = IDMXProtocol::Get(ProtocolName);
+
+	return Protocol;
 }
 
-FName FDMXInputPortConfigCustomization::GetPortGuidPropertyName() const
+FGuid FDMXInputPortConfigCustomization::GetPortGuid() const
 {
-	return FDMXInputPortConfig::GetPortGuidPropertyNameChecked();
+	TArray<void*> RawData;
+	PortGuidHandle->AccessRawData(RawData);
+
+	// Multiediting is not supported, may fire if this is used in a blueprint way that would support it
+	if (ensureMsgf(RawData.Num() == 1, TEXT("Using port config in ways that would enable multiediting is not supported.")))
+	{
+		const FGuid* PortGuidPtr = reinterpret_cast<FGuid*>(RawData[0]);
+		if (PortGuidPtr && PortGuidPtr->IsValid())
+		{
+			return *PortGuidPtr;
+		}
+	}
+
+	return FGuid();
 }
 
 const TArray<EDMXCommunicationType> FDMXInputPortConfigCustomization::GetSupportedCommunicationTypes() const
@@ -74,4 +293,45 @@ const TArray<EDMXCommunicationType> FDMXInputPortConfigCustomization::GetSupport
 	return TArray<EDMXCommunicationType>();
 }
 
+EDMXCommunicationType FDMXInputPortConfigCustomization::GetCommunicationType() const
+{
+	uint8 CommunicationType;
+	ensure(CommunicationTypeHandle->GetValue(CommunicationType) == FPropertyAccess::Success);
+
+	return static_cast<EDMXCommunicationType>(CommunicationType);
+}
+
+FString FDMXInputPortConfigCustomization::GetIPAddress() const
+{
+	FGuid PortGuid = GetPortGuid();
+	if (PortGuid.IsValid())
+	{
+		const UDMXProtocolSettings* ProtocolSettings = GetDefault<UDMXProtocolSettings>();
+		checkf(ProtocolSettings, TEXT("Unexpected protocol settings not available when its details are customized"));
+
+		const FDMXInputPortConfig* InputPortConfigPtr = ProtocolSettings->InputPortConfigs.FindByPredicate([&PortGuid](const FDMXInputPortConfig& InputPortConfig)
+			{
+				return InputPortConfig.GetPortGuid() == PortGuid;
+			});
+		if (InputPortConfigPtr)
+		{
+			return InputPortConfigPtr->GetDeviceAddress();
+		}
+
+		const FDMXOutputPortConfig* OutputPortConfigPtr = ProtocolSettings->OutputPortConfigs.FindByPredicate([&PortGuid](const FDMXOutputPortConfig& OutputPortConfig)
+			{
+				return OutputPortConfig.GetPortGuid() == PortGuid;
+			});
+		if (OutputPortConfigPtr)
+		{
+			return OutputPortConfigPtr->GetDeviceAddress();
+		}
+	}
+
+	FString IPAddress;
+	ensure(DeviceAddressHandle->GetValue(IPAddress) == FPropertyAccess::Success);
+
+	return IPAddress;
+}
+
 #undef LOCTEXT_NAMESPACE
diff --git a/Engine/Plugins/VirtualProduction/DMX/DMXProtocol/Source/DMXProtocolEditor/Private/DetailsCustomizations/DMXInputPortConfigCustomization.h b/Engine/Plugins/VirtualProduction/DMX/DMXProtocol/Source/DMXProtocolEditor/Private/DetailsCustomizations/DMXInputPortConfigCustomization.h
--- a/Engine/Plugins/VirtualProduction/DMX/DMXProtocol/Source/DMXProtocolEditor/Private/DetailsCustomizations/DMXInputPortConfigCustomization.h
+++ b/Engine/Plugins/VirtualProduction/DMX/DMXProtocol/Source/DMXProtocolEditor/Private/DetailsCustomizations/DMXInputPortConfigCustomization.h
@@ -2,28 +2,94 @@
 
 #pragma once
 
-#include "DMXPortConfigCustomizationBase.h"
+#include "DMXProtocolCommon.h"
+
+#include "CoreMinimal.h"
+#include "IPropertyTypeCustomization.h"
+#include "Layout/Visibility.h"
 
 enum class EDMXCommunicationType : uint8;
+class SDMXCommunicationTypeComboBox;
+class SDMXIPAddressEditWidget;
+class SDMXProtocolNameComboBox;
 
 class IDetailPropertyRow;
+class IPropertyHandle;
+class IPropertyUtilities;
 
 
 /** Details customization for input and output port configs. */
 class FDMXInputPortConfigCustomization
-	: public FDMXPortConfigCustomizationBase
+	: public IPropertyTypeCustomization
 {
 public:
 	static TSharedRef<IPropertyTypeCustomization> MakeInstance();
 
-	// ~Begin DMXPortConfigCustomizationBase Interface
-	virtual FName GetProtocolNamePropertyName() const override;
-	virtual FName GetCommunicationTypePropertyName() const override;
-	virtual FName GetDeviceAddressPropertyName() const override;
-	virtual FName GetDestinationAddressPropertyName() const override;
-	virtual FName GetPriorityStrategyPropertyName() const override;
-	virtual FName GetPriorityPropertyName() const override;
-	virtual FName GetPortGuidPropertyName() const override;
-	virtual const TArray<EDMXCommunicationType> GetSupportedCommunicationTypes() const override;
-	// ~End DMXPortConfigCustomizationBase Interface
+protected:
+	// ~Begin IPropertyTypecustomization Interface
+	virtual void CustomizeHeader(TSharedRef<IPropertyHandle> StructPropertyHandle, class FDetailWidgetRow& HeaderRow, IPropertyTypeCustomizationUtils& StructCustomizationUtils) override;
+	virtual void CustomizeChildren(TSharedRef<IPropertyHandle> StructPropertyHandle, class IDetailChildrenBuilder& ChildBuilder, IPropertyTypeCustomizationUtils& StructCustomizationUtils) override;
+	// ~End IPropertyTypecustomization Interface
+
+private:
+	/** Generates the customized Protocol Name row */
+	void GenerateProtocolNameRow(IDetailPropertyRow& PropertyRow);
+
+	/** Generates the customized Communication Type row */
+	void GenerateCommunicationTypeRow(IDetailPropertyRow& PropertyRow);
+
+	/** Generates the customized IP Address row */
+	void GenerateIPAddressRow(IDetailPropertyRow& PropertyRow);
+
+	/** Returns the visibility of the Communication Type Property */
+	EVisibility GetCommunicationTypeVisibility() const;
+
+	/** Called when a Protocol Name was selected */
+	void OnProtocolNameSelected();
+
+	/** Called when a local IP Address was selected */
+	void OnIPAddressSelected();
+
+	/** Called when a Communication Type was selected */
+	void OnCommunicationTypeSelected();
+
+	/** Gets the protocol, checks if it is valid */
+	IDMXProtocolPtr GetProtocol() const;
+
+	/** Helper function that gets the Guid of the edited port */
+	FGuid GetPortGuid() const;
+
+	/** Returns an array of supported communication types */
+	const TArray<EDMXCommunicationType> GetSupportedCommunicationTypes() const;
+
+	/** Gets the communication type */
+	EDMXCommunicationType GetCommunicationType() const;
+
+	/** Gets the IP Address */
+	FString GetIPAddress() const;
+
+	/** Property handle to the ProtocolName property */
+	TSharedPtr<IPropertyHandle> ProtocolNameHandle;
+
+	/** Property handle to the IPAddress property */
+	TSharedPtr<IPropertyHandle> DeviceAddressHandle;
+
+	/** Property handle to the CommunicationType property */
+	TSharedPtr<IPropertyHandle> CommunicationTypeHandle;
+
+	/** Property handle to the PortGuid property */
+	TSharedPtr<IPropertyHandle> PortGuidHandle;
+
+	/** ComboBox to select a protocol name */
+	TSharedPtr<SDMXProtocolNameComboBox> ProtocolNameComboBox;
+
+	/** ComboBox that displays local ip addresses */
+	TSharedPtr<SDMXIPAddressEditWidget> IPAddressEditWidget;
+
+	/** ComboBox that exposes a selection of communication types to the user */
+	TSharedPtr<SDMXCommunicationTypeComboBox> CommunicationTypeComboBox;
+
+	/** Property utilities for this customization */
+	TSharedPtr<IPropertyUtilities> PropertyUtilities;
 };
+
diff --git a/Engine/Plugins/VirtualProduction/DMX/DMXProtocol/Source/DMXProtocolEditor/Private/DetailsCustomizations/DMXOutputPortConfigCustomization.cpp b/Engine/Plugins/VirtualProduction/DMX/DMXProtocol/Source/DMXProtocolEditor/Private/DetailsCustomizations/DMXOutputPortConfigCustomization.cpp
--- a/Engine/Plugins/VirtualProduction/DMX/DMXProtocol/Source/DMXProtocolEditor/Private/DetailsCustomizations/DMXOutputPortConfigCustomization.cpp
+++ b/Engine/Plugins/VirtualProduction/DMX/DMXProtocol/Source/DMXProtocolEditor/Private/DetailsCustomizations/DMXOutputPortConfigCustomization.cpp
@@ -10,59 +10,411 @@
 #include "IO/DMXOutputPortConfig.h"
 #include "IO/DMXPortManager.h"
 #include "Widgets/SDMXCommunicationTypeComboBox.h"
+#include "Widgets/SDMXDelayEditWidget.h"
 #include "Widgets/SDMXIPAddressEditWidget.h"
 #include "Widgets/SDMXProtocolNameComboBox.h"
 
 #include "DetailWidgetRow.h"
 #include "IDetailChildrenBuilder.h"
 #include "IDetailPropertyRow.h" 
-#include "IPAddress.h" 
-#include "SocketSubsystem.h"
+#include "IPropertyUtilities.h"
+#include "ScopedTransaction.h"
 #include "Misc/Guid.h" 
 #include "Widgets/Text/STextBlock.h" 
 
 
-#define LOCTEXT_NAMESPACE "DMXPortConfigCustomization"
+#define LOCTEXT_NAMESPACE "DMXOutputPortConfigCustomization"
 
 TSharedRef<IPropertyTypeCustomization> FDMXOutputPortConfigCustomization::MakeInstance()
 {
 	return MakeShared<FDMXOutputPortConfigCustomization>();
 }
 
-FName FDMXOutputPortConfigCustomization::GetProtocolNamePropertyName() const
+void FDMXOutputPortConfigCustomization::CustomizeHeader(TSharedRef<IPropertyHandle> StructPropertyHandle, class FDetailWidgetRow& HeaderRow, IPropertyTypeCustomizationUtils& StructCustomizationUtils)
 {
-	return FDMXOutputPortConfig::GetProtocolNamePropertyNameChecked();
+	PropertyUtilities = StructCustomizationUtils.GetPropertyUtilities();
+
+	const bool bDisplayResetToDefault = false;
+	const FText DisplayNameOverride = FText::GetEmpty();
+	const FText DisplayToolTipOverride = FText::GetEmpty();
+
+	HeaderRow
+	.NameContent()
+	[
+		StructPropertyHandle->CreatePropertyNameWidget(DisplayNameOverride, DisplayToolTipOverride, bDisplayResetToDefault)
+	];
 }
 
-FName FDMXOutputPortConfigCustomization::GetCommunicationTypePropertyName() const
+void FDMXOutputPortConfigCustomization::CustomizeChildren(TSharedRef<IPropertyHandle> StructPropertyHandle, class IDetailChildrenBuilder& ChildBuilder, IPropertyTypeCustomizationUtils& StructCustomizationUtils)
 {
-	return FDMXOutputPortConfig::GetCommunicationTypePropertyNameChecked();
+	// Retrieve structure's child properties
+	uint32 NumChildren;
+	StructPropertyHandle->GetNumChildren(NumChildren);
+	TMap<FName, TSharedPtr< IPropertyHandle>> PropertyHandles;
+	for (uint32 ChildIndex = 0; ChildIndex < NumChildren; ++ChildIndex)
+	{
+		TSharedRef<IPropertyHandle> ChildHandle = StructPropertyHandle->GetChildHandle(ChildIndex).ToSharedRef();
+		const FName PropertyName = ChildHandle->GetProperty()->GetFName();
+
+		PropertyHandles.Add(PropertyName, ChildHandle);
+	}
+
+	// Cache customized properties
+	ProtocolNameHandle = PropertyHandles.FindChecked(FDMXOutputPortConfig::GetProtocolNamePropertyNameChecked());
+	CommunicationTypeHandle = PropertyHandles.FindChecked(FDMXOutputPortConfig::GetCommunicationTypePropertyNameChecked());
+	DeviceAddressHandle = PropertyHandles.FindChecked(FDMXOutputPortConfig::GetDeviceAddressPropertyNameChecked());
+	DelayHandle = PropertyHandles.FindChecked(FDMXOutputPortConfig::GetDelayPropertyNameChecked());
+	DelayFrameRateHandle = PropertyHandles.FindChecked(FDMXOutputPortConfig::GetDelayFrameRatePropertyNameChecked());
+	PortGuidHandle = PropertyHandles.FindChecked(FDMXOutputPortConfig::GetPortGuidPropertyNameChecked());
+
+	// Ports always need a valid Guid (cannot be blueprinted)
+	if (!GetPortGuid().IsValid())
+	{
+		ChildBuilder.AddCustomRow(LOCTEXT("InvalidPortConfigSearchString", "Invalid"))
+			.WholeRowContent()
+			[
+				SNew(STextBlock)
+				.Text(LOCTEXT("InvalidPortConfigText", "Invalid Port Guid. Cannot utilize this port."))
+			];
+
+		return;
+	}
+
+	// Add customized properties
+	for (auto Iter(PropertyHandles.CreateConstIterator()); Iter; ++Iter)
+	{
+		// Don't add the PortGuid and the DelayFrameRate properties
+		if (Iter.Value() == PortGuidHandle ||
+			Iter.Value() == DelayFrameRateHandle)
+		{
+			continue;
+		}
+		
+		// Add the property
+		IDetailPropertyRow& PropertyRow = ChildBuilder.AddProperty(Iter.Value().ToSharedRef());
+
+		if (Iter.Value() == ProtocolNameHandle)
+		{
+			GenerateProtocolNameRow(PropertyRow);
+		}
+		else if (Iter.Value() == CommunicationTypeHandle)
+		{
+			GenerateCommunicationTypeRow(PropertyRow);
+		}
+		else if (Iter.Value() == DeviceAddressHandle)
+		{
+			GenerateIPAddressRow(PropertyRow);
+		}
+		else if (Iter.Key() == FDMXOutputPortConfig::GetDestinationAddressesPropertyNameChecked())
+		{
+			UpdateDestinationAddressesVisibility();
+
+			TAttribute<EVisibility> DestinationAddressesVisibilityAttribute = TAttribute<EVisibility>::Create(TAttribute<EVisibility>::FGetter::CreateLambda([this]() 
+				{
+					return DestinationAddressesVisibility;
+				}));
+
+			PropertyRow.Visibility(DestinationAddressesVisibilityAttribute);
+		}
+		else if (Iter.Key() == FDMXOutputPortConfig::GetPriorityPropertyNameChecked())
+		{
+			TAttribute<EVisibility> PriorityVisibilityAttribute = TAttribute<EVisibility>::Create(TAttribute<EVisibility>::FGetter::CreateLambda([this]() 
+				{
+					const bool bPriorityIsVisible = GetProtocol().IsValid() && GetProtocol()->SupportsPrioritySettings();
+					return bPriorityIsVisible ? EVisibility::Visible : EVisibility::Collapsed;
+				}));
+
+			PropertyRow.Visibility(PriorityVisibilityAttribute);
+		}
+		else if (Iter.Key() == FDMXOutputPortConfig::GetDelayPropertyNameChecked())
+		{
+			GenerateDelayRow(PropertyRow);
+		}
+	}
 }
 
-FName FDMXOutputPortConfigCustomization::GetDeviceAddressPropertyName() const
+void FDMXOutputPortConfigCustomization::GenerateProtocolNameRow(IDetailPropertyRow& PropertyRow)
 {
-	return FDMXOutputPortConfig::GetDeviceAddressPropertyNameChecked();
+	// Customizate the Protocol Name property to draw a combo box with all protocol names
+
+	TSharedPtr<SWidget> NameWidget;
+	TSharedPtr<SWidget> ValueWidget;
+	FDetailWidgetRow Row;
+	PropertyRow.GetDefaultWidgets(NameWidget, ValueWidget, Row);
+	
+	const FName InitialSelection = [this]() -> FName
+	{
+		if (const IDMXProtocolPtr& Protocol = GetProtocol())
+		{
+			return Protocol->GetProtocolName();
+		}
+
+		return NAME_None;
+	}();
+
+	PropertyRow.CustomWidget()
+		.NameContent()
+		[
+			NameWidget.ToSharedRef()
+		]
+		.ValueContent()
+		[
+			SAssignNew(ProtocolNameComboBox, SDMXProtocolNameComboBox)
+			.InitiallySelectedProtocolName(InitialSelection)
+			.OnProtocolNameSelected(this, &FDMXOutputPortConfigCustomization::OnProtocolNameSelected)
+		];
+
 }
 
-FName FDMXOutputPortConfigCustomization::GetDestinationAddressPropertyName() const
+void FDMXOutputPortConfigCustomization::GenerateCommunicationTypeRow(IDetailPropertyRow& PropertyRow)
 {
-	return FDMXOutputPortConfig::GetDestinationAddressPropertyNameChecked();
+	// Customizate the Communication Type property depending on the port's supported communication modes
+
+	TAttribute<EVisibility> CommunicationTypeVisibilityAttribute = TAttribute<EVisibility>::Create(TAttribute<EVisibility>::FGetter::CreateSP(this, &FDMXOutputPortConfigCustomization::GetCommunicationTypeVisibility));
+	PropertyRow.Visibility(CommunicationTypeVisibilityAttribute);
+
+	TSharedPtr<SWidget> NameWidget;
+	TSharedPtr<SWidget> ValueWidget;
+	FDetailWidgetRow Row;
+	PropertyRow.GetDefaultWidgets(NameWidget, ValueWidget, Row);
+
+	PropertyRow.CustomWidget()
+		.Visibility(CommunicationTypeVisibilityAttribute)
+		.NameContent()
+		[
+			NameWidget.ToSharedRef()
+		]
+		.ValueContent()
+		[
+			SAssignNew(CommunicationTypeComboBox, SDMXCommunicationTypeComboBox)
+			.CommunicationTypes(GetSupportedCommunicationTypes())
+			.InitialCommunicationType(GetCommunicationType())
+			.OnCommunicationTypeSelected(this, &FDMXOutputPortConfigCustomization::OnCommunicationTypeSelected)
+		];
 }
 
-FName FDMXOutputPortConfigCustomization::GetPriorityStrategyPropertyName() const
+void FDMXOutputPortConfigCustomization::GenerateIPAddressRow(IDetailPropertyRow& PropertyRow)
 {
-	// Doesn't exist for output ports
-	return NAME_None;
+	// Customizate the IPAddress property to show a combo box with available IP addresses
+	FString InitialValue = GetIPAddress();
+
+	TSharedPtr<SWidget> NameWidget;
+	TSharedPtr<SWidget> ValueWidget;
+	FDetailWidgetRow Row;
+	PropertyRow.GetDefaultWidgets(NameWidget, ValueWidget, Row);
+
+	PropertyRow.CustomWidget()
+		.NameContent()
+		[
+			NameWidget.ToSharedRef()
+		]
+		.ValueContent()
+		[
+			SAssignNew(IPAddressEditWidget, SDMXIPAddressEditWidget)
+			.InitialValue(InitialValue)
+			.bShowLocalNICComboBox(true)
+			.OnIPAddressSelected(this, &FDMXOutputPortConfigCustomization::OnIPAddressSelected)
+		];
+}	
+
+void FDMXOutputPortConfigCustomization::GenerateDelayRow(IDetailPropertyRow& PropertyRow)
+{
+	// Customizate the IPAddress property to show a combo box with available IP addresses
+	const double InitialDelay = GetDelay();
+	const TArray<FFrameRate> InitialDelayFrameRates = GetDelayFrameRates();
+	
+	FFrameRate InitialDelayFrameRate;
+	if (ensureMsgf(InitialDelayFrameRates.Num() == 1, TEXT("Multi editing Output Port Configs is not supported.")))
+	{
+		InitialDelayFrameRate = InitialDelayFrameRates[0];
+	}
+
+	TSharedPtr<SWidget> NameWidget;
+	TSharedPtr<SWidget> ValueWidget;
+	FDetailWidgetRow Row;
+	PropertyRow.GetDefaultWidgets(NameWidget, ValueWidget, Row);
+
+	PropertyRow.CustomWidget()
+		.NameContent()
+		[
+			NameWidget.ToSharedRef()
+		]
+		.ValueContent()
+		[
+			SAssignNew(DelayEditWidget, SDMXDelayEditWidget)
+			.InitialDelay(InitialDelay)
+			.InitialDelayFrameRate(InitialDelayFrameRate)
+			.OnDelayChanged(this, &FDMXOutputPortConfigCustomization::OnDelayChanged)
+			.OnDelayFrameRateChanged(this, &FDMXOutputPortConfigCustomization::OnDelayFrameRateChanged)
+		];
 }
 
-FName FDMXOutputPortConfigCustomization::GetPriorityPropertyName() const
+EVisibility FDMXOutputPortConfigCustomization::GetCommunicationTypeVisibility() const
 {
-	return FDMXOutputPortConfig::GetPriorityPropertyNameChecked();
+	if (CommunicationTypeComboBox.IsValid() && GetProtocol().IsValid())
+	{
+return CommunicationTypeComboBox->GetVisibility();
+	}
+
+	return EVisibility::Collapsed;
 }
 
-FName FDMXOutputPortConfigCustomization::GetPortGuidPropertyName() const
+void FDMXOutputPortConfigCustomization::OnProtocolNameSelected()
 {
-	return FDMXOutputPortConfig::GetPortGuidPropertyNameChecked();
+	FName ProtocolName = ProtocolNameComboBox->GetSelectedProtocolName();
+
+	const FScopedTransaction Transaction(LOCTEXT("ProtocolSelected", "Select DMX Output Port Protocol"));
+
+	ProtocolNameHandle->NotifyPreChange();
+	ProtocolNameHandle->SetValue(ProtocolName);
+	ProtocolNameHandle->NotifyPostChange(EPropertyChangeType::ValueSet);
+
+	PropertyUtilities->ForceRefresh();
+}
+
+void FDMXOutputPortConfigCustomization::OnCommunicationTypeSelected()
+{
+	EDMXCommunicationType SelectedCommunicationType = CommunicationTypeComboBox->GetSelectedCommunicationType();
+
+	const FScopedTransaction Transaction(LOCTEXT("CommunicationTypeSelected", "Select DMX Output Port Communication Type"));
+
+	CommunicationTypeHandle->NotifyPreChange();
+	CommunicationTypeHandle->SetValue(static_cast<uint8>(SelectedCommunicationType));
+	CommunicationTypeHandle->NotifyPostChange(EPropertyChangeType::ValueSet);
+
+	PropertyUtilities->ForceRefresh();
+}
+
+void FDMXOutputPortConfigCustomization::OnIPAddressSelected()
+{
+	FString SelectedIP = IPAddressEditWidget->GetSelectedIPAddress();
+
+	const FScopedTransaction Transaction(LOCTEXT("CommunicationTypeSelected", "Selected DMX Output Port IP Address"));
+
+	DeviceAddressHandle->NotifyPreChange();
+	ensure(DeviceAddressHandle->SetValue(SelectedIP) == FPropertyAccess::Success);
+	DeviceAddressHandle->NotifyPostChange(EPropertyChangeType::ValueSet);
+}
+
+void FDMXOutputPortConfigCustomization::UpdateDestinationAddressesVisibility()
+{
+	EDMXCommunicationType CommunicationType = GetCommunicationType();
+	if (CommunicationType == EDMXCommunicationType::Unicast)
+	{
+		DestinationAddressesVisibility = EVisibility::Visible;
+	}
+	else
+	{
+		DestinationAddressesVisibility = EVisibility::Collapsed;
+	}
+}
+
+void FDMXOutputPortConfigCustomization::OnDelayChanged()
+{
+	// Clamp the delay value to a reasonable range
+	static const double MaxDelaySeconds = 60.f;
+
+	const double DesiredDelaySeconds = DelayEditWidget->GetDelay() / DelayEditWidget->GetDelayFrameRate().AsDecimal();
+	const double NewDelay = DesiredDelaySeconds > MaxDelaySeconds ? MaxDelaySeconds * DelayEditWidget->GetDelayFrameRate().AsDecimal() : DelayEditWidget->GetDelay();
+
+	const FScopedTransaction Transaction(LOCTEXT("DelalySecondsChange", "Set DMX Output Port Delay"));
+
+	DelayHandle->NotifyPreChange();
+	if (DelayHandle->SetValue(NewDelay) == FPropertyAccess::Success)
+	{
+		DelayHandle->NotifyPostChange(EPropertyChangeType::ValueSet);
+		PropertyUtilities->ForceRefresh();
+	}
+}
+
+void FDMXOutputPortConfigCustomization::OnDelayFrameRateChanged()
+{
+	const FFrameRate DelayFrameRate = DelayEditWidget->GetDelayFrameRate();
+
+	const FScopedTransaction Transaction(LOCTEXT("DelayFrameRateChanged", "Set DMX Output Port Delay Type"));
+
+	DeviceAddressHandle->NotifyPreChange();
+
+	TArray<void*> RawDatas;
+	DelayFrameRateHandle->AccessRawData(RawDatas);
+
+	for (void* RawData : RawDatas)
+	{
+		FFrameRate* FrameRatePtr = (FFrameRate*)RawData;
+		*FrameRatePtr = DelayFrameRate;
+	}
+	DeviceAddressHandle->NotifyPostChange(EPropertyChangeType::ValueSet);
+
+	// Clamp the delay value to a reasonable range
+	static const double MaxDelaySeconds = 60.f;
+
+	const double DesiredDelaySeconds = DelayEditWidget->GetDelay() / DelayFrameRate.AsDecimal();
+	const double NewDelay = DesiredDelaySeconds > MaxDelaySeconds ? MaxDelaySeconds * DelayFrameRate.AsDecimal() : DelayEditWidget->GetDelay();
+
+	DelayHandle->NotifyPreChange();
+	if (DelayHandle->SetValue(NewDelay) == FPropertyAccess::Success)
+	{
+		DelayHandle->NotifyPostChange(EPropertyChangeType::ValueSet);
+		PropertyUtilities->ForceRefresh();
+	}
+}
+
+double FDMXOutputPortConfigCustomization::GetDelay() const
+{
+	double Delay;
+	if (DelayHandle->GetValue(Delay) == FPropertyAccess::Success)
+	{
+		return Delay;
+	}
+
+	return 0.0;
+}
+
+TArray<FFrameRate> FDMXOutputPortConfigCustomization::GetDelayFrameRates() const
+{
+	TArray<FFrameRate> FrameRates;
+
+	TArray<void*> RawDatas;
+	DelayFrameRateHandle->AccessRawData(RawDatas);
+
+	for (void* RawData : RawDatas)
+	{
+		const FFrameRate* FrameRatePtr = (FFrameRate*)RawData;
+		if (FrameRatePtr)
+		{
+			FrameRates.Add(*FrameRatePtr);
+		}
+	}
+
+	return FrameRates;
+}
+
+IDMXProtocolPtr FDMXOutputPortConfigCustomization::GetProtocol() const
+{
+	FName ProtocolName;
+	ProtocolNameHandle->GetValue(ProtocolName);
+
+	IDMXProtocolPtr Protocol = IDMXProtocol::Get(ProtocolName);
+	
+	return Protocol;
+}
+
+FGuid FDMXOutputPortConfigCustomization::GetPortGuid() const
+{
+	TArray<void*> RawData;
+	PortGuidHandle->AccessRawData(RawData);
+
+	// Multiediting is not supported, may fire if this is used in a blueprint way that would support it
+	if (ensureMsgf(RawData.Num() == 1, TEXT("Using port config in ways that would enable multiediting is not supported.")))
+	{
+		const FGuid* PortGuidPtr = reinterpret_cast<FGuid*>(RawData[0]);
+		if (PortGuidPtr && PortGuidPtr->IsValid())
+		{
+			return *PortGuidPtr;
+		}
+	}
+
+	return FGuid();
 }
 
 const TArray<EDMXCommunicationType> FDMXOutputPortConfigCustomization::GetSupportedCommunicationTypes() const
@@ -75,4 +427,45 @@ const TArray<EDMXCommunicationType> FDMXOutputPortConfigCustomization::GetSuppor
 	return TArray<EDMXCommunicationType>();
 }
 
+EDMXCommunicationType FDMXOutputPortConfigCustomization::GetCommunicationType() const
+{
+	uint8 CommunicationType;
+	ensure(CommunicationTypeHandle->GetValue(CommunicationType) == FPropertyAccess::Success);
+
+	return static_cast<EDMXCommunicationType>(CommunicationType);
+}
+
+FString FDMXOutputPortConfigCustomization::GetIPAddress() const
+{
+	FGuid PortGuid = GetPortGuid();
+	if (PortGuid.IsValid())
+	{
+		const UDMXProtocolSettings* ProtocolSettings = GetDefault<UDMXProtocolSettings>();
+		checkf(ProtocolSettings, TEXT("Unexpected protocol settings not available when its details are customized"));
+
+		const FDMXInputPortConfig* InputPortConfigPtr = ProtocolSettings->InputPortConfigs.FindByPredicate([&PortGuid](const FDMXInputPortConfig& InputPortConfig)
+			{
+				return InputPortConfig.GetPortGuid() == PortGuid;
+			});
+		if (InputPortConfigPtr)
+		{
+			return InputPortConfigPtr->GetDeviceAddress();
+		}
+
+		const FDMXOutputPortConfig* OutputPortConfigPtr = ProtocolSettings->OutputPortConfigs.FindByPredicate([&PortGuid](const FDMXOutputPortConfig& OutputPortConfig)
+			{
+				return OutputPortConfig.GetPortGuid() == PortGuid;
+			});
+		if (OutputPortConfigPtr)
+		{
+			return OutputPortConfigPtr->GetDeviceAddress();
+		}
+	}
+
+	FString IPAddress;
+	ensure(DeviceAddressHandle->GetValue(IPAddress) == FPropertyAccess::Success);
+
+	return IPAddress;
+}
+
 #undef LOCTEXT_NAMESPACE
diff --git a/Engine/Plugins/VirtualProduction/DMX/DMXProtocol/Source/DMXProtocolEditor/Private/DetailsCustomizations/DMXOutputPortConfigCustomization.h b/Engine/Plugins/VirtualProduction/DMX/DMXProtocol/Source/DMXProtocolEditor/Private/DetailsCustomizations/DMXOutputPortConfigCustomization.h
--- a/Engine/Plugins/VirtualProduction/DMX/DMXProtocol/Source/DMXProtocolEditor/Private/DetailsCustomizations/DMXOutputPortConfigCustomization.h
+++ b/Engine/Plugins/VirtualProduction/DMX/DMXProtocol/Source/DMXProtocolEditor/Private/DetailsCustomizations/DMXOutputPortConfigCustomization.h
@@ -2,29 +2,125 @@
 
 #pragma once
 
-#include "DMXPortConfigCustomizationBase.h"
+#include "DMXProtocolCommon.h"
+
+#include "CoreMinimal.h"
+#include "IPropertyTypeCustomization.h"
+#include "Layout/Visibility.h"
 
 enum class EDMXCommunicationType : uint8;
+class SDMXCommunicationTypeComboBox;
+class SDMXDelayEditWidget;
+class SDMXIPAddressEditWidget;
+class SDMXProtocolNameComboBox;
 
+struct FFrameRate;
 class IDetailPropertyRow;
+class IPropertyHandle;
+class IPropertyUtilities;
 
 
 /** Details customization for input and output port configs. */
 class FDMXOutputPortConfigCustomization
-	: public FDMXPortConfigCustomizationBase
+	: public IPropertyTypeCustomization
 {
 public:
 	static TSharedRef<IPropertyTypeCustomization> MakeInstance();
 
 protected:
-	// ~Begin DMXPortConfigCustomizationBase Interface
-	virtual FName GetProtocolNamePropertyName() const override;
-	virtual FName GetCommunicationTypePropertyName() const override;
-	virtual FName GetDeviceAddressPropertyName() const override;
-	virtual FName GetDestinationAddressPropertyName() const override;
-	virtual FName GetPriorityStrategyPropertyName() const override;
-	virtual FName GetPriorityPropertyName() const override;
-	virtual FName GetPortGuidPropertyName() const override;
-	virtual const TArray<EDMXCommunicationType> GetSupportedCommunicationTypes() const override;
-	// ~End DMXPortConfigCustomizationBase Interface
+	// ~Begin IPropertyTypecustomization Interface
+	virtual void CustomizeHeader(TSharedRef<IPropertyHandle> StructPropertyHandle, class FDetailWidgetRow& HeaderRow, IPropertyTypeCustomizationUtils& StructCustomizationUtils) override;
+	virtual void CustomizeChildren(TSharedRef<IPropertyHandle> StructPropertyHandle, class IDetailChildrenBuilder& ChildBuilder, IPropertyTypeCustomizationUtils& StructCustomizationUtils) override;
+	// ~End IPropertyTypecustomization Interface
+
+private:
+	/** Generates the customized Protocol Name row */
+	void GenerateProtocolNameRow(IDetailPropertyRow& PropertyRow);
+
+	/** Generates the customized Communication Type row */
+	void GenerateCommunicationTypeRow(IDetailPropertyRow& PropertyRow);
+
+	/** Generates the customized IP Address row */
+	void GenerateIPAddressRow(IDetailPropertyRow& PropertyRow);
+
+	/** Generates the customized Delay Seconds row */
+	void GenerateDelayRow(IDetailPropertyRow& PropertyRow);
+
+	/** Returns the visibility of the Communication Type Property */
+	EVisibility GetCommunicationTypeVisibility() const;
+
+	/** Called when a Protocol Name was selected */
+	void OnProtocolNameSelected();
+
+	/** Called when a local IP Address was selected */
+	void OnIPAddressSelected();
+
+	/** Called when a Communication Type was selected */
+	void OnCommunicationTypeSelected();
+
+	/** Called when the Destination Addresses Array visibility needs to be updated */
+	void UpdateDestinationAddressesVisibility();
+
+	/** Called when the Delay Seconds changed in the Delay Edit Widget */
+	void OnDelayChanged();
+
+	/** Called when the Delay Frame Rate changed in the Delay Edit Widget */
+	void OnDelayFrameRateChanged();
+
+	/** Returns the delay of the port */
+	double GetDelay() const;
+
+	/** Returns the delay frame rate of the port */
+	TArray<FFrameRate> GetDelayFrameRates() const;
+
+	/** Gets the protocol, checks if it is valid */
+	IDMXProtocolPtr GetProtocol() const;
+
+	/** Helper function that gets the Guid of the edited port */
+	FGuid GetPortGuid() const;
+
+	/** Returns an array of supported communication types */
+	const TArray<EDMXCommunicationType> GetSupportedCommunicationTypes() const;
+
+	/** Gets the communication type */
+	EDMXCommunicationType GetCommunicationType() const;
+
+	/** Gets the IP Address */
+	FString GetIPAddress() const;
+
+	/** Property handle to the ProtocolName property */
+	TSharedPtr<IPropertyHandle> ProtocolNameHandle;
+
+	/** Property handle to the IPAddress property */
+	TSharedPtr<IPropertyHandle> DeviceAddressHandle;
+
+	/** Property handle to the CommunicationType property */
+	TSharedPtr<IPropertyHandle> CommunicationTypeHandle;
+
+	/** Property handle to the Delay property */
+	TSharedPtr<IPropertyHandle> DelayHandle;
+
+	/** Property handle to the DelayFrameRate property */
+	TSharedPtr<IPropertyHandle> DelayFrameRateHandle;
+
+	/** Property handle to the PortGuid property */
+	TSharedPtr<IPropertyHandle> PortGuidHandle;
+
+	/** Visibility of the destination addresses array */
+	EVisibility DestinationAddressesVisibility;
+
+	/** ComboBox to select a protocol name */
+	TSharedPtr<SDMXProtocolNameComboBox> ProtocolNameComboBox;
+
+	/** ComboBox that displays local ip addresses */
+	TSharedPtr<SDMXIPAddressEditWidget> IPAddressEditWidget;
+
+	/** ComboBox that exposes a selection of communication types to the user */
+	TSharedPtr<SDMXCommunicationTypeComboBox> CommunicationTypeComboBox;
+
+	/** Custom widget that lets the user set the delay */
+	TSharedPtr<SDMXDelayEditWidget> DelayEditWidget;
+
+	/** Property utilities for this customization */
+	TSharedPtr<IPropertyUtilities> PropertyUtilities;
 };
diff --git a/Engine/Plugins/VirtualProduction/DMX/DMXProtocol/Source/DMXProtocolEditor/Private/DetailsCustomizations/DMXOutputPortDestinationAddressCustomization.cpp b/Engine/Plugins/VirtualProduction/DMX/DMXProtocol/Source/DMXProtocolEditor/Private/DetailsCustomizations/DMXOutputPortDestinationAddressCustomization.cpp
new file mode 100644
--- /dev/null
+++ b/Engine/Plugins/VirtualProduction/DMX/DMXProtocol/Source/DMXProtocolEditor/Private/DetailsCustomizations/DMXOutputPortDestinationAddressCustomization.cpp
@@ -0,0 +1,77 @@
+// Copyright Epic Games, Inc. All Rights Reserved.
+
+#include "DMXOutputPortDestinationAddressCustomization.h"
+
+#include "IO/DMXOutputPortConfig.h"
+#include "Widgets/SDMXIPAddressEditWidget.h"
+
+#include "IDetailChildrenBuilder.h"
+#include "IDetailPropertyRow.h" 
+#include "DetailWidgetRow.h"
+
+
+#define LOCTEXT_NAMESPACE "DMXOutputPortDestinationAddressCustomization"
+
+TSharedRef<IPropertyTypeCustomization> FDMXOutputPortDestinationAddressCustomization::MakeInstance()
+{
+	return MakeShared<FDMXOutputPortDestinationAddressCustomization>();
+}
+
+void FDMXOutputPortDestinationAddressCustomization::CustomizeHeader(TSharedRef<IPropertyHandle> StructPropertyHandle, class FDetailWidgetRow& HeaderRow, IPropertyTypeCustomizationUtils& StructCustomizationUtils)
+{
+	// Nothing in the header
+}
+
+void FDMXOutputPortDestinationAddressCustomization::CustomizeChildren(TSharedRef<IPropertyHandle> StructPropertyHandle, class IDetailChildrenBuilder& ChildBuilder, IPropertyTypeCustomizationUtils& StructCustomizationUtils)
+{
+	DestinationAddressStringHandle = StructPropertyHandle->GetChildHandle(GET_MEMBER_NAME_CHECKED(FDMXOutputPortDestinationAddress, DestinationAddressString));
+
+	ChildBuilder.AddCustomRow(LOCTEXT("InvalidPortConfigSearchString", "Invalid"))
+		.ValueContent()
+		[
+			SAssignNew(IPAddressEditWidget, SDMXIPAddressEditWidget)
+			.InitialValue(GetIPAddress())
+			.bShowLocalNICComboBox(false)
+			.OnIPAddressSelected(this, &FDMXOutputPortDestinationAddressCustomization::OnIPAddressSelected)
+		];
+}
+
+FString FDMXOutputPortDestinationAddressCustomization::GetIPAddress() const
+{
+	FString IPAddress;
+	if (DestinationAddressStringHandle->GetValue(IPAddress) == FPropertyAccess::Success)
+	{
+		return IPAddress;
+	}
+
+	return TEXT("");
+}
+
+void FDMXOutputPortDestinationAddressCustomization::SetIPAddress(const FString& NewIPAddress)
+{
+	DestinationAddressStringHandle->NotifyPreChange();
+	DestinationAddressStringHandle->SetValue(NewIPAddress);
+	DestinationAddressStringHandle->NotifyPostChange(EPropertyChangeType::ValueSet);
+}
+
+void FDMXOutputPortDestinationAddressCustomization::OnIPAddressSelected()
+{
+	TArray<UObject*> OuterObjects;
+	DestinationAddressStringHandle->GetOuterObjects(OuterObjects);
+	
+	// Forward the change to the outer object so it gets the property change. Required due to get the change from the nested property all the way to UDMXProtocolSettings.
+	for (UObject* Object : OuterObjects)
+	{
+		Object->PreEditChange(FDMXOutputPortDestinationAddress::StaticStruct()->FindPropertyByName(GET_MEMBER_NAME_CHECKED(FDMXOutputPortDestinationAddress, DestinationAddressString)));
+	}
+
+	const FString IPAddress = IPAddressEditWidget->GetSelectedIPAddress();
+	SetIPAddress(IPAddress);
+
+	for (UObject* Object : OuterObjects)
+	{
+		Object->PostEditChange();
+	}
+}
+
+#undef LOCTEXT_NAMESPACE
diff --git a/Engine/Plugins/VirtualProduction/DMX/DMXProtocol/Source/DMXProtocolEditor/Private/DetailsCustomizations/DMXOutputPortDestinationAddressCustomization.h b/Engine/Plugins/VirtualProduction/DMX/DMXProtocol/Source/DMXProtocolEditor/Private/DetailsCustomizations/DMXOutputPortDestinationAddressCustomization.h
new file mode 100644
--- /dev/null
+++ b/Engine/Plugins/VirtualProduction/DMX/DMXProtocol/Source/DMXProtocolEditor/Private/DetailsCustomizations/DMXOutputPortDestinationAddressCustomization.h
@@ -0,0 +1,41 @@
+// Copyright Epic Games, Inc. All Rights Reserved.
+
+#pragma once
+
+#include "CoreMinimal.h"
+#include "IPropertyTypeCustomization.h"
+
+class SDMXIPAddressEditWidget;
+
+class IPropertyHandle;
+
+
+/** Customization for the DMXOutputPortConfigCustomization struct */
+class FDMXOutputPortDestinationAddressCustomization
+	: public IPropertyTypeCustomization
+{
+public:
+	static TSharedRef<IPropertyTypeCustomization> MakeInstance();
+
+protected:
+	// ~Begin IPropertyTypecustomization Interface
+	virtual void CustomizeHeader(TSharedRef<IPropertyHandle> StructPropertyHandle, class FDetailWidgetRow& HeaderRow, IPropertyTypeCustomizationUtils& StructCustomizationUtils) override;
+	virtual void CustomizeChildren(TSharedRef<IPropertyHandle> StructPropertyHandle, class IDetailChildrenBuilder& ChildBuilder, IPropertyTypeCustomizationUtils& StructCustomizationUtils) override;
+	// ~End IPropertyTypecustomization Interface
+
+private:
+	/** Gets the IP Address currently set */
+	FString GetIPAddress() const;
+
+	/** Sets the IP Adderss */
+	void SetIPAddress(const FString& NewIPAddress);
+
+	/** Called when an IP Address was selcted */
+	void OnIPAddressSelected();
+
+	/** Widget to enter an IP address */
+	TSharedPtr<SDMXIPAddressEditWidget> IPAddressEditWidget;
+
+	/** Property handle for the DestinationAddressString property */
+	TSharedPtr<IPropertyHandle> DestinationAddressStringHandle;
+};
diff --git a/Engine/Plugins/VirtualProduction/DMX/DMXProtocol/Source/DMXProtocolEditor/Private/DetailsCustomizations/DMXPortConfigCustomizationBase .cpp b/Engine/Plugins/VirtualProduction/DMX/DMXProtocol/Source/DMXProtocolEditor/Private/DetailsCustomizations/DMXPortConfigCustomizationBase .cpp
deleted file mode 100644
diff --git a/Engine/Plugins/VirtualProduction/DMX/DMXProtocol/Source/DMXProtocolEditor/Private/DetailsCustomizations/DMXPortConfigCustomizationBase.h b/Engine/Plugins/VirtualProduction/DMX/DMXProtocol/Source/DMXProtocolEditor/Private/DetailsCustomizations/DMXPortConfigCustomizationBase.h
deleted file mode 100644
diff --git a/Engine/Plugins/VirtualProduction/DMX/DMXProtocol/Source/DMXProtocolEditor/Private/DetailsCustomizations/Widgets/SDMXDelayEditWidget.cpp b/Engine/Plugins/VirtualProduction/DMX/DMXProtocol/Source/DMXProtocolEditor/Private/DetailsCustomizations/Widgets/SDMXDelayEditWidget.cpp
new file mode 100644
--- /dev/null
+++ b/Engine/Plugins/VirtualProduction/DMX/DMXProtocol/Source/DMXProtocolEditor/Private/DetailsCustomizations/Widgets/SDMXDelayEditWidget.cpp
@@ -0,0 +1,229 @@
+// Copyright Epic Games, Inc. All Rights Reserved.
+
+#include "SDMXDelayEditWidget.h"
+
+#include "DMXProtocolSettings.h"
+
+#include "CommonFrameRates.h"
+#include "EditorStyleSet.h"
+#include "Framework/MultiBox/MultiBoxBuilder.h"
+#include "Widgets/SBoxPanel.h"
+#include "Widgets/SFrameRateEntryBox.h"
+#include "Widgets/Input/SEditableTextBox.h"
+
+
+#define LOCTEXT_NAMESPACE "SDMXDelayEditWidget"
+
+void SDMXDelayEditWidget::Construct(const FArguments& InArgs)
+{
+	OnDelayChanged = InArgs._OnDelayChanged;
+	OnDelayFrameRateChanged = InArgs._OnDelayFrameRateChanged;
+
+	Delay = InArgs._InitialDelay;
+	DelayFrameRate = InArgs._InitialDelayFrameRate;
+
+	const FText InitialDelayText = FText::FromString(FString::SanitizeFloat(Delay));
+	
+	const TAttribute<EVisibility> FPSVisiblityAttribute = TAttribute<EVisibility>::Create(TAttribute<EVisibility>::FGetter::CreateLambda([this]()
+		{
+			return DelayFrameRate.AsDecimal() == 1.f ? EVisibility::Collapsed : EVisibility::Visible;
+		}));
+
+
+	ChildSlot
+	[
+		SNew(SHorizontalBox)
+		+ SHorizontalBox::Slot()
+		.HAlign(HAlign_Left)
+		.VAlign(VAlign_Center)
+		[
+			SAssignNew(ValueEditBox, SEditableTextBox)
+			.Text(InitialDelayText)
+			.OnVerifyTextChanged(this, &SDMXDelayEditWidget::OnVerifyValueTextChanged)
+			.OnTextCommitted(this, &SDMXDelayEditWidget::OnValueTextCommitted)
+		]
+		
+	+ SHorizontalBox::Slot()
+		.HAlign(HAlign_Left)
+		.VAlign(VAlign_Center)
+		.Padding(8.f, 0.f, 8.f, 0.f)
+		.AutoWidth()
+		[
+			SNew(STextBlock)
+			.Text_Lambda([this]
+				{
+					if (Delay == 1.0)
+					{
+						return LOCTEXT("FramesPerSecondLabelSingular", "frame at");
+					}
+					else
+					{
+						return LOCTEXT("FramesPerSecondLabelPlural", "frames at");
+					}
+				})
+			.Visibility(FPSVisiblityAttribute)
+		]
+
+		+ SHorizontalBox::Slot()
+		.HAlign(HAlign_Left)
+		.VAlign(VAlign_Center)
+		.AutoWidth()
+		[
+			SNew(SComboButton)
+			.OnGetMenuContent(this, &SDMXDelayEditWidget::GenereateDelayTypeMenu)
+			.ButtonContent()
+			[
+				SNew(STextBlock)
+				.Text_Lambda([this]()
+					{
+						const double DelayFrameRateDecimal = DelayFrameRate.AsDecimal();
+						if (DelayFrameRateDecimal == 1.0)
+						{
+							return LOCTEXT("SecondsLabel", "Seconds");
+						}
+						else
+						{
+							return FText::Format(LOCTEXT("FramesPerSecond", "{0} Frames Per Second"), DelayFrameRateDecimal);
+						}
+					})
+			]
+		]
+	];
+}
+
+double SDMXDelayEditWidget::GetDelay() const
+{
+	return Delay;
+}
+
+FFrameRate SDMXDelayEditWidget::GetDelayFrameRate() const
+{
+	return DelayFrameRate;
+}
+
+TSharedRef<SWidget> SDMXDelayEditWidget::GenereateDelayTypeMenu()
+{
+	FMenuBuilder MenuBuilder(true, nullptr);
+
+	const FFrameRate OneFramePerSecond = FFrameRate(1, 1);
+	MenuBuilder.AddMenuEntry(
+		LOCTEXT("Seconds", "As Seconds"),
+		LOCTEXT("SecondsDescription", "Sets the delay to Seconds"),
+		FSlateIcon(),
+		FUIAction(
+			FExecuteAction::CreateSP(this, &SDMXDelayEditWidget::OnDelayTypeChanged, OneFramePerSecond))
+		);
+
+	MenuBuilder.AddMenuSeparator();
+
+	TArray<FCommonFrameRateInfo> CommonFrameRates;
+	for (const FCommonFrameRateInfo& Info : FCommonFrameRates::GetAll())
+	{
+		CommonFrameRates.Add(Info);
+	}
+
+	CommonFrameRates.Sort(
+		[](const FCommonFrameRateInfo& A, const FCommonFrameRateInfo& B)
+		{
+			return A.FrameRate.AsDecimal() < B.FrameRate.AsDecimal();
+		}
+	);
+
+	MenuBuilder.BeginSection(NAME_None, LOCTEXT("RecommendedRates", "As Frame Rate"));
+	{
+		for (const FCommonFrameRateInfo& Info : CommonFrameRates)
+		{
+			MenuBuilder.AddMenuEntry(
+				Info.DisplayName,
+				Info.Description,
+				FSlateIcon(),
+				FUIAction(
+					FExecuteAction::CreateSP(this, &SDMXDelayEditWidget::OnDelayTypeChanged, Info.FrameRate),
+					FCanExecuteAction(),
+					FIsActionChecked::CreateSP(this, &SDMXDelayEditWidget::IsSameDelayType, Info.FrameRate)
+				),
+				NAME_None,
+				EUserInterfaceActionType::RadioButton
+			);
+		}
+
+		MenuBuilder.AddWidget(
+			SNew(SBox)
+			.HAlign(HAlign_Right)
+			.MaxDesiredWidth(100.f)
+			[
+				SNew(SFrameRateEntryBox)
+				.Value_Lambda([this]
+					{
+						return DelayFrameRate;
+					})
+				.OnValueChanged(this, &SDMXDelayEditWidget::OnDelayTypeChanged)
+			],
+			LOCTEXT("CustomFramerateDisplayLabel", "Custom")
+		);
+	}
+	MenuBuilder.EndSection();
+
+	return MenuBuilder.MakeWidget();
+}
+
+void SDMXDelayEditWidget::OnDelayTypeChanged(FFrameRate NewDelayFrameRate)
+{
+	DelayFrameRate = NewDelayFrameRate;
+
+	OnDelayFrameRateChanged.ExecuteIfBound();
+}
+
+bool SDMXDelayEditWidget::IsSameDelayType(FFrameRate FrameRate) const
+{
+	if (DelayFrameRate == FrameRate)
+	{
+		return true;
+	}
+
+	return false;
+}
+
+bool SDMXDelayEditWidget::OnVerifyValueTextChanged(const FText& InNewText, FText& OutErrorMessage)
+{
+	FString StringValue = InNewText.ToString();
+	int32 Value;
+	if (LexTryParseString<int32>(Value, *StringValue))
+	{
+		if (Value >= 0.0)
+		{
+			OutErrorMessage = FText::GetEmpty();
+			return true;
+		}
+	}
+
+	OutErrorMessage = LOCTEXT("InvalidValueError", "Only numeric values > 0 are supported");
+	return false;
+}
+
+void SDMXDelayEditWidget::OnValueTextCommitted(const FText& InNewText, ETextCommit::Type InTextCommit)
+{
+	FString StringValue = InNewText.ToString();
+	double Value;
+	if (LexTryParseString<double>(Value, *StringValue))
+	{
+		if (Value >= 0.0)
+		{
+			Delay = Value;
+
+			const FText ValueText = FText::FromString(FString::SanitizeFloat(Delay));
+			ValueEditBox->SetText(ValueText);
+
+			OnDelayChanged.ExecuteIfBound();
+		}
+	}
+}
+
+bool SDMXDelayEditWidget::DisplayAsFramesPerSecond() const
+{
+	const double Value = DelayFrameRate.AsDecimal() / Delay;
+
+	return FMath::IsNearlyEqual(Value, 1.0);
+}
+
+#undef LOCTEXT_NAMESPACE
diff --git a/Engine/Plugins/VirtualProduction/DMX/DMXProtocol/Source/DMXProtocolEditor/Private/DetailsCustomizations/Widgets/SDMXDelayEditWidget.h b/Engine/Plugins/VirtualProduction/DMX/DMXProtocol/Source/DMXProtocolEditor/Private/DetailsCustomizations/Widgets/SDMXDelayEditWidget.h
new file mode 100644
--- /dev/null
+++ b/Engine/Plugins/VirtualProduction/DMX/DMXProtocol/Source/DMXProtocolEditor/Private/DetailsCustomizations/Widgets/SDMXDelayEditWidget.h
@@ -0,0 +1,78 @@
+// Copyright Epic Games, Inc. All Rights Reserved.
+
+#pragma once
+
+#include "CoreMinimal.h"
+
+#include "Misc/FrameRate.h"
+#include "Widgets/SCompoundWidget.h"
+#include "Widgets/Input/SComboBox.h"
+
+class SEditableTextBox;
+
+
+/** Helper widget to select a delay either by Seconds or by a Frame Rate */
+class SDMXDelayEditWidget
+	: public SCompoundWidget
+{
+public:
+	SLATE_BEGIN_ARGS(SDMXDelayEditWidget)
+		{}
+
+		/** The initially displayed delay in seconds */
+		SLATE_ARGUMENT(double, InitialDelay)
+
+		/** The initially selected frame rate (1.0 for seconds) */
+		SLATE_ARGUMENT(FFrameRate, InitialDelayFrameRate)
+
+		/** Delegate executed when the delay seconds changed */
+		SLATE_EVENT(FSimpleDelegate, OnDelayChanged)
+
+		/** Delegate executed when the delay frame rate changed */
+		SLATE_EVENT(FSimpleDelegate, OnDelayFrameRateChanged)
+
+	SLATE_END_ARGS()
+
+	/** Constructs the widget */
+	void Construct(const FArguments& InArgs);
+
+	/** Returns the delay in seconds */
+	double GetDelay() const;
+
+	/** Returns the delay frame rate */
+	FFrameRate GetDelayFrameRate() const;
+
+private:
+	/** Generates the menu to select a Type */
+	TSharedRef<SWidget> GenereateDelayTypeMenu();
+
+	/** Called when the value text changed but was not commited */
+	bool OnVerifyValueTextChanged(const FText& InNewText, FText& OutErrorMessage);
+
+	/** Called when the value text was commited */
+	void OnValueTextCommitted(const FText& InNewText, ETextCommit::Type InTextCommit);
+
+	/** Called when the delay Type changed  */
+	void OnDelayTypeChanged(FFrameRate NewDelayFrameRate);
+
+	/** Called to test if the delay Type is the same as the selected one */
+	bool IsSameDelayType(FFrameRate FrameRate) const;
+
+	/** Returns true if values should be displayed as frames per second */
+	bool DisplayAsFramesPerSecond() const;
+
+	/** Box to enter the Delay value */
+	TSharedPtr<SEditableTextBox> ValueEditBox;
+
+	/** Delegate executed when the selected delay changed */
+	FSimpleDelegate OnDelayChanged;
+
+	/** Delegate executed when the selected delay changed */
+	FSimpleDelegate OnDelayFrameRateChanged;
+
+	/** The currently set delay, in seconds */
+	double Delay = 0.0;
+
+	/** The frame rate in which the delay is displayed */
+	FFrameRate DelayFrameRate;
+};
diff --git a/Engine/Plugins/VirtualProduction/DMX/DMXProtocol/Source/DMXProtocolEditor/Private/DetailsCustomizations/Widgets/SDMXIPAddressEditWidget.cpp b/Engine/Plugins/VirtualProduction/DMX/DMXProtocol/Source/DMXProtocolEditor/Private/DetailsCustomizations/Widgets/SDMXIPAddressEditWidget.cpp
--- a/Engine/Plugins/VirtualProduction/DMX/DMXProtocol/Source/DMXProtocolEditor/Private/DetailsCustomizations/Widgets/SDMXIPAddressEditWidget.cpp
+++ b/Engine/Plugins/VirtualProduction/DMX/DMXProtocol/Source/DMXProtocolEditor/Private/DetailsCustomizations/Widgets/SDMXIPAddressEditWidget.cpp
@@ -7,6 +7,8 @@
 #include "DMXProtocolUtils.h"
 
 #include "EditorStyleSet.h"
+#include "IPAddress.h"
+#include "SocketSubsystem.h"
 #include "Widgets/Layout/SWidgetSwitcher.h"
 #include "Widgets/Input/SEditableTextBox.h"
 
@@ -17,20 +19,32 @@ void SDMXIPAddressEditWidget::Construct(const FArguments& InArgs)
 
 	LocalAdapterAddressSource = FDMXProtocolUtils::GetLocalNetworkInterfaceCardIPs();
 
-	ChildSlot
-	[
-		SAssignNew(LocalAdapterAddressComboBox, SComboBox<TSharedPtr<FString>>)
-		.OptionsSource(&LocalAdapterAddressSource)
-		.OnGenerateWidget(this, &SDMXIPAddressEditWidget::GenerateLocalAdapterAddressComboBoxEntry)
-		.OnSelectionChanged(this, &SDMXIPAddressEditWidget::OnIpAddressSelected)
-		.Content()
+	const TSharedRef<SEditableTextBox> IPAddressEditTextBox = SAssignNew(IPAddressEditableTextBlock, SEditableTextBox)
+		.Font(FEditorStyle::GetFontStyle(TEXT("PropertyWindow.NormalFont")))
+		.Text(FText::FromString(InArgs._InitialValue))
+		.OnTextCommitted(this, &SDMXIPAddressEditWidget::OnIPAddressTextCommmited);
+
+	if (InArgs._bShowLocalNICComboBox)
+	{
+		ChildSlot
+		[
+			SAssignNew(LocalAdapterAddressComboBox, SComboBox<TSharedPtr<FString>>)
+			.OptionsSource(&LocalAdapterAddressSource)
+			.OnGenerateWidget(this, &SDMXIPAddressEditWidget::GenerateLocalAdapterAddressComboBoxEntry)
+			.OnSelectionChanged(this, &SDMXIPAddressEditWidget::OnIpAddressSelected)
+			.Content()
+			[
+				IPAddressEditTextBox
+			]
+		];
+	}
+	else
+	{
+		ChildSlot
 		[
-			SAssignNew(IPAddressEditableTextBlock, SEditableTextBox)
-			.Font(FEditorStyle::GetFontStyle(TEXT("PropertyWindow.NormalFont")))
-			.Text(FText::FromString(InArgs._InitialValue))
-			.OnTextCommitted(this, &SDMXIPAddressEditWidget::OnIPAddressEntered)
-		]
-	];
+			IPAddressEditTextBox
+		];
+	}
 }
 
 FString SDMXIPAddressEditWidget::GetSelectedIPAddress() const
@@ -48,7 +62,7 @@ void SDMXIPAddressEditWidget::OnIpAddressSelected(TSharedPtr<FString> InAddress,
 	}
 }
 
-void SDMXIPAddressEditWidget::OnIPAddressEntered(const FText&, ETextCommit::Type)
+void SDMXIPAddressEditWidget::OnIPAddressTextCommmited(const FText&, ETextCommit::Type)
 {
 	OnIPAddressSelectedDelegate.ExecuteIfBound();
 }
diff --git a/Engine/Plugins/VirtualProduction/DMX/DMXProtocol/Source/DMXProtocolEditor/Private/DetailsCustomizations/Widgets/SDMXIPAddressEditWidget.h b/Engine/Plugins/VirtualProduction/DMX/DMXProtocol/Source/DMXProtocolEditor/Private/DetailsCustomizations/Widgets/SDMXIPAddressEditWidget.h
--- a/Engine/Plugins/VirtualProduction/DMX/DMXProtocol/Source/DMXProtocolEditor/Private/DetailsCustomizations/Widgets/SDMXIPAddressEditWidget.h
+++ b/Engine/Plugins/VirtualProduction/DMX/DMXProtocol/Source/DMXProtocolEditor/Private/DetailsCustomizations/Widgets/SDMXIPAddressEditWidget.h
@@ -20,10 +20,13 @@ class SDMXIPAddressEditWidget
 {
 public:
 	SLATE_BEGIN_ARGS(SDMXIPAddressEditWidget)
+		: _bShowLocalNICComboBox(false)
 		{}
 
 		SLATE_ARGUMENT(FString, InitialValue)
 
+		SLATE_ARGUMENT(bool, bShowLocalNICComboBox)
+
 		SLATE_EVENT(FSimpleDelegate, OnIPAddressSelected)
 
 	SLATE_END_ARGS()
@@ -39,8 +42,8 @@ private:
 	/** Handles changes in the local adapter address combo box */
 	void OnIpAddressSelected(TSharedPtr<FString> InAddress, ESelectInfo::Type InType);
 
-	/** Called when the ip address was entered in the editable text block */
-	void OnIPAddressEntered(const FText&, ETextCommit::Type);
+	/** Called when the IP Address was commited in the editable text block */
+	void OnIPAddressTextCommmited(const FText&, ETextCommit::Type);
 
 	/** Generates an entry in the local adapter address combo box */
 	TSharedRef<SWidget> GenerateLocalAdapterAddressComboBoxEntry(TSharedPtr<FString> InAddress);
diff --git a/Engine/Plugins/VirtualProduction/DMX/DMXProtocol/Source/DMXProtocolSACN/Private/DMXProtocolSACN.cpp b/Engine/Plugins/VirtualProduction/DMX/DMXProtocol/Source/DMXProtocolSACN/Private/DMXProtocolSACN.cpp
--- a/Engine/Plugins/VirtualProduction/DMX/DMXProtocol/Source/DMXProtocolSACN/Private/DMXProtocolSACN.cpp
+++ b/Engine/Plugins/VirtualProduction/DMX/DMXProtocol/Source/DMXProtocolSACN/Private/DMXProtocolSACN.cpp
@@ -148,7 +148,7 @@ void FDMXProtocolSACN::UnregisterInputPort(const TSharedRef<FDMXInputPort, ESPMo
 	}
 }
 
-TSharedPtr<IDMXSender> FDMXProtocolSACN::RegisterOutputPort(const TSharedRef<FDMXOutputPort, ESPMode::ThreadSafe>& OutputPort)
+TArray<TSharedPtr<IDMXSender>> FDMXProtocolSACN::RegisterOutputPort(const TSharedRef<FDMXOutputPort, ESPMode::ThreadSafe>& OutputPort)
 {
 	check(!OutputPort->IsRegistered());
 	check(!CachedOutputPorts.Contains(OutputPort));
@@ -157,26 +157,33 @@ TSharedPtr<IDMXSender> FDMXProtocolSACN::RegisterOutputPort(const TSharedRef<FDM
 	EDMXCommunicationType CommunicationType = OutputPort->GetCommunicationType();
 
 	// Try to use an existing receiver or create a new one
-	TSharedPtr<FDMXProtocolSACNSender> Sender = nullptr;
+	TArray<TSharedPtr<IDMXSender>> NewsSenders;
 	
 	if (CommunicationType == EDMXCommunicationType::Multicast)
 	{
-		Sender = FindExistingMulticastSender(NetworkInterfaceAddress);
-			
-		if (!Sender.IsValid())
+		TSharedPtr<FDMXProtocolSACNSender> NewSender = FindExistingMulticastSender(NetworkInterfaceAddress);
+		if (!NewSender.IsValid())
 		{
-			Sender = FDMXProtocolSACNSender::TryCreateMulticastSender(SharedThis(this), NetworkInterfaceAddress);
+			NewSender = FDMXProtocolSACNSender::TryCreateMulticastSender(SharedThis(this), NetworkInterfaceAddress);
+			Senders.Add(NewSender);
 		}
+
+		NewSender->AssignOutputPort(OutputPort);
+		NewsSenders.Add(NewSender);
 	}
 	else if (CommunicationType == EDMXCommunicationType::Unicast)
 	{
-		const FString& UnicastAddress = OutputPort->GetDestinationAddress();
-
-		Sender = FindExistingUnicastSender(NetworkInterfaceAddress, UnicastAddress);
-
-		if (!Sender.IsValid())
+		for (const FString& UnicastAddress : OutputPort->GetDestinationAddresses())
 		{
-			Sender = FDMXProtocolSACNSender::TryCreateUnicastSender(SharedThis(this), NetworkInterfaceAddress, UnicastAddress);
+			TSharedPtr<FDMXProtocolSACNSender> NewSender = FindExistingUnicastSender(NetworkInterfaceAddress, UnicastAddress);
+			if (!NewSender.IsValid())
+			{
+				NewSender = FDMXProtocolSACNSender::TryCreateUnicastSender(SharedThis(this), NetworkInterfaceAddress, UnicastAddress);
+				Senders.Add(NewSender);
+			}
+
+			NewSender->AssignOutputPort(OutputPort);
+			NewsSenders.Add(NewSender);
 		}
 	}
 	else
@@ -184,21 +191,15 @@ TSharedPtr<IDMXSender> FDMXProtocolSACN::RegisterOutputPort(const TSharedRef<FDM
 		// Invalid Communication Type
 		UE_LOG(LogDMXProtocol, Error, TEXT("Cannot create DMX Protocol sACN Sender. The communication type specified is not supported."));
 	}
-	
 
-	if (!Sender.IsValid())
+	if (NewsSenders.Num() == 0)
 	{
 		UE_LOG(LogDMXProtocol, Warning, TEXT("Could not create sACN sender for output port %s"), *OutputPort->GetPortName());
-
-		return nullptr;
 	}
 
-	Senders.Add(Sender);
-
-	Sender->AssignOutputPort(OutputPort);
 	CachedOutputPorts.Add(OutputPort);
 
-	return Sender;
+	return NewsSenders;
 }
 
 void FDMXProtocolSACN::UnregisterOutputPort(const TSharedRef<FDMXOutputPort, ESPMode::ThreadSafe>& OutputPort)
@@ -207,7 +208,7 @@ void FDMXProtocolSACN::UnregisterOutputPort(const TSharedRef<FDMXOutputPort, ESP
 	CachedOutputPorts.Remove(OutputPort);
 
 	TSharedPtr<FDMXProtocolSACNSender> UnusedSender;
-	for (const TSharedPtr<FDMXProtocolSACNSender>& Sender : Senders)
+	for (const TSharedPtr<FDMXProtocolSACNSender>& Sender : TSet<TSharedPtr<FDMXProtocolSACNSender>>(Senders))
 	{
 		if (Sender->ContainsOutputPort(OutputPort))
 		{
@@ -215,16 +216,10 @@ void FDMXProtocolSACN::UnregisterOutputPort(const TSharedRef<FDMXOutputPort, ESP
 
 			if (Sender->GetNumAssignedOutputPorts() == 0)
 			{
-				UnusedSender = Sender;
+				Senders.Remove(Sender);
 			}
-			break;
 		}
 	}
-
-	if (UnusedSender.IsValid())
-	{
-		Senders.Remove(UnusedSender);
-	}
 }
 
 bool FDMXProtocolSACN::IsCausingLoopback(EDMXCommunicationType InCommunicationType)
diff --git a/Engine/Plugins/VirtualProduction/DMX/DMXProtocol/Source/DMXProtocolSACN/Private/DMXProtocolSACN.h b/Engine/Plugins/VirtualProduction/DMX/DMXProtocol/Source/DMXProtocolSACN/Private/DMXProtocolSACN.h
--- a/Engine/Plugins/VirtualProduction/DMX/DMXProtocol/Source/DMXProtocolSACN/Private/DMXProtocolSACN.h
+++ b/Engine/Plugins/VirtualProduction/DMX/DMXProtocol/Source/DMXProtocolSACN/Private/DMXProtocolSACN.h
@@ -38,7 +38,7 @@ public:
 	virtual bool SupportsPrioritySettings() const override;
 	virtual bool RegisterInputPort(const TSharedRef<FDMXInputPort, ESPMode::ThreadSafe>& InputPort) override;
 	virtual void UnregisterInputPort(const TSharedRef<FDMXInputPort, ESPMode::ThreadSafe>& InputPort) override;
-	virtual TSharedPtr<IDMXSender> RegisterOutputPort(const TSharedRef<FDMXOutputPort, ESPMode::ThreadSafe>& OutputPort) override;
+	virtual TArray<TSharedPtr<IDMXSender>>  RegisterOutputPort(const TSharedRef<FDMXOutputPort, ESPMode::ThreadSafe>& OutputPort) override;
 	virtual bool IsCausingLoopback(EDMXCommunicationType InCommunicationType) override;
 	virtual void UnregisterOutputPort(const TSharedRef<FDMXOutputPort, ESPMode::ThreadSafe>& OutputPort) override;
 	// ~End IDMXProtocol implementation
diff --git a/Engine/Plugins/VirtualProduction/DMX/DMXProtocol/Source/DMXProtocolSACN/Private/DMXProtocolSACNReceiver.cpp b/Engine/Plugins/VirtualProduction/DMX/DMXProtocol/Source/DMXProtocolSACN/Private/DMXProtocolSACNReceiver.cpp
--- a/Engine/Plugins/VirtualProduction/DMX/DMXProtocol/Source/DMXProtocolSACN/Private/DMXProtocolSACNReceiver.cpp
+++ b/Engine/Plugins/VirtualProduction/DMX/DMXProtocol/Source/DMXProtocolSACN/Private/DMXProtocolSACNReceiver.cpp
@@ -5,6 +5,7 @@
 #include "DMXProtocolLog.h"
 #include "DMXProtocolSACN.h"
 #include "DMXProtocolSACNConstants.h"
+#include "DMXProtocolSACNUtils.h"
 #include "DMXStats.h"
 #include "IO/DMXInputPort.h"
 #include "Packets/DMXProtocolE131PDUPacket.h"
@@ -140,7 +141,7 @@ void FDMXProtocolSACNReceiver::AssignInputPort(const TSharedPtr<FDMXInputPort, E
 			ISocketSubsystem* SocketSubsystem = ISocketSubsystem::Get(PLATFORM_SOCKETSUBSYSTEM);
 			TSharedRef<FInternetAddr> NewMulticastGroupAddr = SocketSubsystem->CreateInternetAddr();
 
-			uint32 MulticastIp = GetIpForUniverseID(UniverseID);
+			uint32 MulticastIp = FDMXProtocolSACNUtils::GetIpForUniverseID(UniverseID);
 			NewMulticastGroupAddr->SetIp(MulticastIp);
 			NewMulticastGroupAddr->SetPort(ACN_PORT);
 
@@ -270,7 +271,7 @@ void FDMXProtocolSACNReceiver::Update(const FTimespan& SocketWaitTime)
 					Reader->Seek(0);
 
 					// validate Universe
-					const uint32 FakeMulticastIp = GetIpForUniverseID(UniverseID);
+					const uint32 FakeMulticastIp = FDMXProtocolSACNUtils::GetIpForUniverseID(UniverseID);
 					if (MulticastGroupAddrToUniverseIDMap.Contains(FakeMulticastIp))
 					{
 						DistributeReceivedData(UniverseID, Reader);
@@ -356,7 +357,7 @@ void FDMXProtocolSACNReceiver::HandleDataPacket(uint16 UniverseID, const TShared
 	// Copy relevant data
 	FMemory::Memcpy(UniverseData->GetData() + IncomingDMXDMPLayer.FirstPropertyAddress, IncomingDMXDMPLayer.DMX.GetData(), IncomingDMXDMPLayer.PropertyValueCount - 1);
 	
-	FDMXSignalSharedRef DMXSignal = MakeShared<FDMXSignal, ESPMode::ThreadSafe>(FPlatformTime::Seconds(), UniverseID, IncomingDMXFramingLayer.Priority, PropertiesCacheValues[UniverseID], 0);
+	FDMXSignalSharedRef DMXSignal = MakeShared<FDMXSignal, ESPMode::ThreadSafe>(FPlatformTime::Seconds(), UniverseID, IncomingDMXFramingLayer.Priority, PropertiesCacheValues[UniverseID]);
 	for (const TSharedPtr<FDMXInputPort, ESPMode::ThreadSafe>& InputPort : AssignedInputPorts)
 	{
 		// Check packet priority
diff --git a/Engine/Plugins/VirtualProduction/DMX/DMXProtocol/Source/DMXProtocolSACN/Private/DMXProtocolSACNReceiver.h b/Engine/Plugins/VirtualProduction/DMX/DMXProtocol/Source/DMXProtocolSACN/Private/DMXProtocolSACNReceiver.h
--- a/Engine/Plugins/VirtualProduction/DMX/DMXProtocol/Source/DMXProtocolSACN/Private/DMXProtocolSACNReceiver.h
+++ b/Engine/Plugins/VirtualProduction/DMX/DMXProtocol/Source/DMXProtocolSACN/Private/DMXProtocolSACNReceiver.h
@@ -58,6 +58,7 @@ public:
 	FORCEINLINE const TSet<TSharedPtr<FDMXInputPort, ESPMode::ThreadSafe>>& GetAssignedInputPorts() const { return AssignedInputPorts; }
 
 	/** Returns the IP address of a universe */
+	UE_DEPRECATED(5.0, "Deprecated in favor of the more generic FDMXProtocolSACNUtils::GetIPForUniverseID")
 	static uint32 GetIpForUniverseID(uint16 InUniverseID);
 
 private:
@@ -100,7 +101,7 @@ private:
 	TSharedPtr<FDMXProtocolSACN, ESPMode::ThreadSafe> Protocol;
 
 	/** The network socket. */
-	FSocket* Socket;
+	FSocket* Socket = nullptr;
 
 	/** The endpoint internet addr */
 	TSharedPtr<FInternetAddr> EndpointInternetAddr;
@@ -109,7 +110,7 @@ private:
 	TAtomic<bool> bStopping;
 
 	/** The thread object. */
-	FRunnableThread* Thread;
+	FRunnableThread* Thread = nullptr;
 
 	/** The receiver thread's name. */
 	FString ThreadName;
diff --git a/Engine/Plugins/VirtualProduction/DMX/DMXProtocol/Source/DMXProtocolSACN/Private/DMXProtocolSACNSender.cpp b/Engine/Plugins/VirtualProduction/DMX/DMXProtocol/Source/DMXProtocolSACN/Private/DMXProtocolSACNSender.cpp
--- a/Engine/Plugins/VirtualProduction/DMX/DMXProtocol/Source/DMXProtocolSACN/Private/DMXProtocolSACNSender.cpp
+++ b/Engine/Plugins/VirtualProduction/DMX/DMXProtocol/Source/DMXProtocolSACN/Private/DMXProtocolSACNSender.cpp
@@ -4,7 +4,7 @@
 
 #include "DMXProtocolSACN.h"
 #include "DMXProtocolSACNReceiver.h"
-
+#include "DMXProtocolSACNUtils.h"
 #include "DMXProtocolConstants.h"
 #include "DMXProtocolLog.h"
 #include "DMXProtocolSettings.h"
@@ -16,7 +16,6 @@
 #include "SocketSubsystem.h"
 #include "Common/UdpSocketBuilder.h"
 #include "Common/UdpSocketSender.h"
-#include "HAL/RunnableThread.h"
 #include "Serialization/ArrayReader.h"
 #include "Serialization/ArrayWriter.h"
 #include "UObject/Class.h"
@@ -50,15 +49,10 @@ FDMXProtocolSACNSender::FDMXProtocolSACNSender(const TSharedPtr<FDMXProtocolSACN
 	, Socket(&InSocket)
 	, NetworkInterfaceInternetAddr(InNetworkInterfaceInternetAddr)
 	, DestinationInternetAddr(InDestinationInternetAddr)
-	, bStopping(false)
-	, Thread(nullptr)
 	, bIsMulticast(bInIsMulticast)
 {
 	check(DestinationInternetAddr.IsValid());
 
-	FString SenderThreadName = FString(TEXT("sACNSender_")) + InDestinationInternetAddr->ToString(false);
-	Thread = FRunnableThread::Create(this, *SenderThreadName, 0U, TPri_TimeCritical, FPlatformAffinity::GetPoolThreadMask());
-
 	if (bIsMulticast)
 	{
 		UE_LOG(LogDMXProtocol, VeryVerbose, TEXT("Created sACN Multicast Sender at %s"), *NetworkInterfaceInternetAddr->ToString(false));
@@ -77,12 +71,6 @@ FDMXProtocolSACNSender::~FDMXProtocolSACNSender()
 		SocketSubsystem->DestroySocket(Socket);
 	}
 
-	if (Thread != nullptr)
-	{
-		Thread->Kill(true);
-		delete Thread;
-	}
-
 	UE_LOG(LogDMXProtocol, VeryVerbose, TEXT("Destroyed sACN Sender at %s sending to %s"), *NetworkInterfaceInternetAddr->ToString(false), *DestinationInternetAddr->ToString(false));
 }
 
@@ -198,195 +186,67 @@ bool FDMXProtocolSACNSender::IsCausingLoopback() const
 
 void FDMXProtocolSACNSender::SendDMXSignal(const FDMXSignalSharedRef& DMXSignal)
 {
-	Buffer.Enqueue(DMXSignal);
-}
-
-void FDMXProtocolSACNSender::ClearBuffer()
-{
-	FScopeLock Lock(&LatestSignalLock);
+	TArray<uint8> Packet;
 
-	Buffer.Empty();
-	UniverseToLatestSignalMap.Reset();
-}
+	FDMXProtocolE131RootLayerPacket RootLayer;
+	static FGuid Guid = FGuid::NewGuid();
+	FMemory::Memcpy(RootLayer.CID.GetData(), &Guid, ACN_CIDBYTES);
 
-bool FDMXProtocolSACNSender::Init()
-{
-	return true;
-}
+	Packet.Append(*RootLayer.Pack(ACN_DMX_SIZE));
 
-uint32 FDMXProtocolSACNSender::Run()
-{
-	const UDMXProtocolSettings* DMXSettings = GetDefault<UDMXProtocolSettings>();
-	check(DMXSettings);
-	
-	// Fixed rate delta time
-	const double SendDeltaTime = 1.f / DMXSettings->SendingRefreshRate;
+	int32 UniverseID = DMXSignal->ExternUniverseID;
 
-	while (!bStopping)
-	{
-		const double StartTime = FPlatformTime::Seconds();
+	FDMXProtocolE131FramingLayerPacket FramingLayer;
+	FramingLayer.Universe = UniverseID;
+	FramingLayer.SequenceNumber = UniverseIDToSequenceNumberMap.FindOrAdd(UniverseID, -1)++; // Init to max, let it wrap over to 0 at first
+	FramingLayer.Priority = DMXSignal->Priority;
+	Packet.Append(*FramingLayer.Pack(ACN_DMX_SIZE));
 
-		Update();
+	FDMXProtocolE131DMPLayerPacket DMPLayer;
+	DMPLayer.AddressIncrement = ACN_ADDRESS_INC;
+	DMPLayer.PropertyValueCount = ACN_DMX_SIZE + 1;
+	FMemory::Memcpy(DMPLayer.DMX.GetData(), DMXSignal->ChannelData.GetData(), ACN_DMX_SIZE);
 
-		const double EndTime = FPlatformTime::Seconds();
-		const double WaitTime = SendDeltaTime - (EndTime - StartTime);
+	Packet.Append(*DMPLayer.Pack(ACN_DMX_SIZE));
 
-		if (WaitTime > 0.f)
-		{
-			// Sleep by the amount which is set in refresh rate
-			FPlatformProcess::SleepNoStats(WaitTime);
-		}
+	const int32 SendDataSize = Packet.Num();
 
-		// In the unlikely case we took to long to send, we instantly continue, but do not take 
-		// further measures to compensate - We would have to run faster than DMX send rate to catch up.
-	}
-
-	return 0;
-}
-
-void FDMXProtocolSACNSender::Stop()
-{
-	bStopping = true;
-}
-
-void FDMXProtocolSACNSender::Exit()
-{
-}
-
-void FDMXProtocolSACNSender::Tick()
-{
-	Update();
-}
-
-FSingleThreadRunnable* FDMXProtocolSACNSender::GetSingleThreadInterface()
-{
-	return this;
-}
+	// Try to send, log errors but avoid spaming the Log
+	static bool bErrorEverLogged = false;
+	TSharedPtr<FInternetAddr> CurrentDestination = DestinationInternetAddr;
 
-void FDMXProtocolSACNSender::Update()
-{
-	// process delayed signals
-	const double Now = FPlatformTime::Seconds();
-	
-	FScopeLock Lock(&LatestSignalLock);
-
-	for (;;)
+	// if in multicast, compute the destination
+	if (bIsMulticast)
 	{
-		// we can safely assumes that older messages are on tail of the queue
-		FDMXSignal OldestDMXSignal;
-		if (DelayedBuffer.Peek(OldestDMXSignal))
-		{
-			if ((OldestDMXSignal.Timestamp + static_cast<double>(OldestDMXSignal.Delay)) <= Now)
-			{
-				if (UniverseToLatestSignalMap.Contains(OldestDMXSignal.ExternUniverseID))
-				{
-					UniverseToLatestSignalMap[OldestDMXSignal.ExternUniverseID] = MakeShared<FDMXSignal>(OldestDMXSignal);
-				}
-				else
-				{
-					UniverseToLatestSignalMap.Add(OldestDMXSignal.ExternUniverseID, MakeShared<FDMXSignal>(OldestDMXSignal));
-				}
-				DelayedBuffer.Pop();
-				continue;
-			}
-		}
-		break;
+		ISocketSubsystem* SocketSubsystem = ISocketSubsystem::Get(PLATFORM_SOCKETSUBSYSTEM);
+		CurrentDestination = SocketSubsystem->CreateInternetAddr();
+
+		uint32 IpForUniverse = FDMXProtocolSACNUtils::GetIpForUniverseID(UniverseID);
+		CurrentDestination->SetIp(IpForUniverse);
+		CurrentDestination->SetPort(ACN_PORT);
 	}
 
+	int32 BytesSent = -1;
+	if (Socket->SendTo(Packet.GetData(), Packet.Num(), BytesSent, *CurrentDestination))
 	{
-		// Keep latest signal per universe
-		FDMXSignalSharedPtr DequeuedDMXSignal;
-		while (Buffer.Dequeue(DequeuedDMXSignal))
-		{
-			if (Protocol->IsValidUniverseID(DequeuedDMXSignal->ExternUniverseID))
-			{
-				if (DequeuedDMXSignal->Delay > 0)
-				{
-					// we make a copy here of the signal for not losing track
-					DelayedBuffer.Enqueue(*DequeuedDMXSignal);
-					continue;
-				}
-				if (UniverseToLatestSignalMap.Contains(DequeuedDMXSignal->ExternUniverseID))
-				{
-					UniverseToLatestSignalMap[DequeuedDMXSignal->ExternUniverseID] = DequeuedDMXSignal.ToSharedRef();
-				}
-				else
-				{
-					UniverseToLatestSignalMap.Add(DequeuedDMXSignal->ExternUniverseID, DequeuedDMXSignal.ToSharedRef());
-				}
-			}
-		}
+		INC_DWORD_STAT(STAT_SACNPackagesSent);
 	}
-
-	// Create a packet for each universe and send it
-	for (const TTuple<int32, FDMXSignalSharedRef>& UniverseToSignalKvp : UniverseToLatestSignalMap)
+	else
 	{
-		int32 UniverseID = UniverseToSignalKvp.Key;
-		const FDMXSignalSharedRef& DMXSignal = UniverseToSignalKvp.Value;
-
-		TArray<uint8> Packet;
-
-		FDMXProtocolE131RootLayerPacket RootLayer;
-		static FGuid Guid = FGuid::NewGuid();
-		FMemory::Memcpy(RootLayer.CID.GetData(), &Guid, ACN_CIDBYTES);
-
-		Packet.Append(*RootLayer.Pack(ACN_DMX_SIZE));
-
-		FDMXProtocolE131FramingLayerPacket FramingLayer;
-		FramingLayer.Universe = UniverseID;
-		FramingLayer.SequenceNumber = UniverseIDToSequenceNumberMap.FindOrAdd(UniverseID, -1)++; // Init to max, let it wrap over to 0 at first
-		FramingLayer.Priority = DMXSignal->Priority;
-		Packet.Append(*FramingLayer.Pack(ACN_DMX_SIZE));
-
-		FDMXProtocolE131DMPLayerPacket DMPLayer;
-		DMPLayer.AddressIncrement = ACN_ADDRESS_INC;
-		DMPLayer.PropertyValueCount = ACN_DMX_SIZE + 1;
-		FMemory::Memcpy(DMPLayer.DMX.GetData(), DMXSignal->ChannelData.GetData(), ACN_DMX_SIZE);
-
-		Packet.Append(*DMPLayer.Pack(ACN_DMX_SIZE));
-
-		const int32 SendDataSize = Packet.Num();
-		int32 BytesSent = -1;
-
-		// Try to send, log errors but avoid spaming the Log
-		static bool bErrorEverLogged = false;
-		TSharedPtr<FInternetAddr> CurrentDestination = DestinationInternetAddr;
-
-		// if in multicast, compute the destination
-		if (bIsMulticast)
+		if (!bErrorEverLogged)
 		{
 			ISocketSubsystem* SocketSubsystem = ISocketSubsystem::Get(PLATFORM_SOCKETSUBSYSTEM);
-			CurrentDestination = SocketSubsystem->CreateInternetAddr();
+			TEnumAsByte<ESocketErrors> RecvFromError = SocketSubsystem->GetLastErrorCode();
 
-			uint32 IpForUniverse = FDMXProtocolSACNReceiver::GetIpForUniverseID(UniverseID);
-			CurrentDestination->SetIp(IpForUniverse);
-			CurrentDestination->SetPort(ACN_PORT);
-		}
+			UE_LOG(LogDMXProtocol, Error, TEXT("Failed send DMX to %s with Error Code %d"), *DestinationInternetAddr->ToString(false), RecvFromError);
 
-		if (Socket->SendTo(Packet.GetData(), Packet.Num(), BytesSent, *CurrentDestination))
-		{
-			INC_DWORD_STAT(STAT_SACNPackagesSent);
+			bErrorEverLogged = true;
 		}
-		else
-		{
-			if(!bErrorEverLogged)
-			{
-				ISocketSubsystem* SocketSubsystem = ISocketSubsystem::Get(PLATFORM_SOCKETSUBSYSTEM);
-				TEnumAsByte<ESocketErrors> RecvFromError = SocketSubsystem->GetLastErrorCode();
-
-				UE_LOG(LogDMXProtocol, Error, TEXT("Failed send DMX to %s with Error Code %d"), *DestinationInternetAddr->ToString(false), RecvFromError);
+	}
 
-				bErrorEverLogged = true;
-			}
-		}
-		
-		if (BytesSent != SendDataSize)
-		{
-			if (!bErrorEverLogged)
-			{
-				UE_LOG(LogDMXProtocol, Warning, TEXT("Incomplete DMX Packet sent to %s"), *DestinationInternetAddr->ToString(false));
-				bErrorEverLogged = true;
-			}
-		}
+	if (BytesSent != SendDataSize && !bErrorEverLogged)
+	{
+		UE_LOG(LogDMXProtocol, Warning, TEXT("Incomplete DMX Packet sent to %s"), *DestinationInternetAddr->ToString(false));
+		bErrorEverLogged = true;
 	}
 }
diff --git a/Engine/Plugins/VirtualProduction/DMX/DMXProtocol/Source/DMXProtocolSACN/Private/DMXProtocolSACNSender.h b/Engine/Plugins/VirtualProduction/DMX/DMXProtocol/Source/DMXProtocolSACN/Private/DMXProtocolSACNSender.h
--- a/Engine/Plugins/VirtualProduction/DMX/DMXProtocol/Source/DMXProtocolSACN/Private/DMXProtocolSACNSender.h
+++ b/Engine/Plugins/VirtualProduction/DMX/DMXProtocol/Source/DMXProtocolSACN/Private/DMXProtocolSACNSender.h
@@ -7,13 +7,7 @@
 #include "Interfaces/IDMXSender.h"
 
 #include "CoreMinimal.h"
-#include "Containers/Queue.h"
-#include "HAL/CriticalSection.h"
-#include "HAL/Runnable.h"
-#include "Misc/ScopeLock.h" 
-#include "Misc/SingleThreadRunnable.h"
 #include "Serialization/ArrayWriter.h"
-#include "Templates/Atomic.h"
 
 
 struct FDMXProtocolSACNRDM;
@@ -23,15 +17,12 @@ class FDMXOutputPort;
 
 class FDMXProtocolSACN;
 class FInternetAddr;
-class FRunnableThread;
 class FSocket;
 class ISocketSubsystem;
 
 
 class DMXPROTOCOLSACN_API FDMXProtocolSACNSender
-	: public FRunnable
-	, public FSingleThreadRunnable
-	, public IDMXSender
+	: public IDMXSender
 {
 protected:
 	FDMXProtocolSACNSender& operator=(const FDMXProtocolSACNSender&) = delete;
@@ -67,7 +58,6 @@ public:
 
 	// ~Begin IDMXSender Interface
 	virtual void SendDMXSignal(const FDMXSignalSharedRef& DMXSignal) override;
-	virtual void ClearBuffer() override;
 	// ~End IDMXSender Interface
 
 public:
@@ -96,30 +86,6 @@ private:
 	/** The Output ports the receiver uses */
 	TSet<TSharedPtr<FDMXOutputPort, ESPMode::ThreadSafe>> AssignedOutputPorts;
 
-protected:
-	//~ Begin FRunnable implementation
-	virtual bool Init() override;
-	virtual uint32 Run() override;
-	virtual void Stop() override;
-	virtual void Exit() override;
-	//~ End FRunnable implementation
-
-	//~ Begin FSingleThreadRunnable implementation
-	virtual void Tick() override;
-	virtual class FSingleThreadRunnable* GetSingleThreadInterface() override;
-	//~ End FSingleThreadRunnable implementation
-
-protected:
-	/** Updates the thread, sending DMX */
-	void Update();
-
-private:
-	/** Buffer of dmx signals */
-	TQueue<FDMXSignalSharedPtr> Buffer;
-
-	/** Map of the latest signal per universe */
-	TMap<int32 /** Universe */, FDMXSignalSharedRef> UniverseToLatestSignalMap;
-
 	/** Map of universes with their current sequence number */
 	TMap<uint16 /** Universe ID */, uint16 /** Sequence Number */> UniverseIDToSequenceNumberMap;
 
@@ -127,7 +93,7 @@ private:
 	TSharedPtr<FDMXProtocolSACN, ESPMode::ThreadSafe> Protocol;
 
 	/** Holds the network socket used to sender packages. */
-	FSocket* Socket;
+	FSocket* Socket = nullptr;
 
 	/** The network interface internet addr */
 	TSharedPtr<FInternetAddr> NetworkInterfaceInternetAddr;
@@ -138,15 +104,6 @@ private:
 	/** Communication type used for the network traffic */
 	EDMXCommunicationType CommunicationType;
 
-	/** Lock used when the buffer is cleared */
-	FCriticalSection LatestSignalLock;
-
-	/** Flag indicating that the thread is stopping. */
-	TAtomic<bool> bStopping;
-
-	/** Holds the thread object. */
-	FRunnableThread* Thread;
-
 	/** Is a Multicast sender ? */
 	const bool bIsMulticast;
 
diff --git a/Engine/Plugins/VirtualProduction/DMX/DMXProtocol/Source/DMXProtocolSACN/Private/DMXProtocolSACNUtils.cpp b/Engine/Plugins/VirtualProduction/DMX/DMXProtocol/Source/DMXProtocolSACN/Private/DMXProtocolSACNUtils.cpp
new file mode 100644
--- /dev/null
+++ b/Engine/Plugins/VirtualProduction/DMX/DMXProtocol/Source/DMXProtocolSACN/Private/DMXProtocolSACNUtils.cpp
@@ -0,0 +1,26 @@
+// Copyright Epic Games, Inc. All Rights Reserved.
+
+#include "DMXProtocolSACNUtils.h"
+
+#include "IPAddress.h"
+#include "SocketSubsystem.h"
+#include "Serialization/BufferArchive.h"
+
+
+uint32 FDMXProtocolSACNUtils::GetIpForUniverseID(uint16 InUniverseID)
+{
+	ISocketSubsystem* SocketSubsystem = ISocketSubsystem::Get(PLATFORM_SOCKETSUBSYSTEM);
+	TSharedPtr<FInternetAddr> InternetAddr = SocketSubsystem->CreateInternetAddr();
+	uint32 ReturnAddress = 0;
+	FBufferArchive IP;
+	uint8 IP_0 = ACN_UNIVERSE_IP_0;
+	uint8 IP_1 = ACN_UNIVERSE_IP_1;
+	IP << IP_0; // [x.?.?.?]
+	IP << IP_1; // [x.x.?.?]
+	IP.ByteSwap(&InUniverseID, sizeof(uint16));
+	IP << InUniverseID;	// [x.x.x.x]
+
+	InternetAddr->SetRawIp(IP);
+	InternetAddr->GetIp(ReturnAddress);
+	return ReturnAddress;
+}
diff --git a/Engine/Plugins/VirtualProduction/DMX/DMXProtocol/Source/DMXProtocolSACN/Private/DMXProtocolSACNUtils.h b/Engine/Plugins/VirtualProduction/DMX/DMXProtocol/Source/DMXProtocolSACN/Private/DMXProtocolSACNUtils.h
--- a/Engine/Plugins/VirtualProduction/DMX/DMXProtocol/Source/DMXProtocolSACN/Private/DMXProtocolSACNUtils.h
+++ b/Engine/Plugins/VirtualProduction/DMX/DMXProtocol/Source/DMXProtocolSACN/Private/DMXProtocolSACNUtils.h
@@ -8,24 +8,13 @@
 
 #include "Serialization/ArrayReader.h"
 
-namespace SACN
-{
-	static uint32 GetRootPacketType(const FArrayReaderPtr& Buffer)
-	{
-
-		uint32 Vector = 0x00000000;
-		const uint32 MinCheck = ACN_ADDRESS_ROOT_VECTOR + 4;
-		if (Buffer->Num() > MinCheck)
-		{
-			// Get OpCode
-			Buffer->Seek(ACN_ADDRESS_ROOT_VECTOR);
-			*Buffer << Vector;
-			Buffer->ByteSwap(&Vector, 4);
 
-			// Reset Position
-			Buffer->Seek(0);
-		}
+class FDMXProtocolSACNUtils
+{
+public:
+	/** Returns the SACN specific IP for a Universe ID */
+	static uint32 GetIpForUniverseID(uint16 InUniverseID);
 
-		return Vector;
-	}
-}
\ No newline at end of file
+	// Can't instantiate this class
+	FDMXProtocolSACNUtils() = delete;
+};
