DMX - Fix CIS NonUnity Compile LyraEditor Win64

diff --git a/Engine/Plugins/VirtualProduction/DMX/DMXProtocol/Source/DMXProtocolSACN/Private/DMXProtocolSACNSender.h b/Engine/Plugins/VirtualProduction/DMX/DMXProtocol/Source/DMXProtocolSACN/Private/DMXProtocolSACNSender.h
--- a/Engine/Plugins/VirtualProduction/DMX/DMXProtocol/Source/DMXProtocolSACN/Private/DMXProtocolSACNSender.h
+++ b/Engine/Plugins/VirtualProduction/DMX/DMXProtocol/Source/DMXProtocolSACN/Private/DMXProtocolSACNSender.h
@@ -106,7 +106,4 @@ private:
 
 	/** Is a Multicast sender ? */
 	const bool bIsMulticast;
-
-	/** Buffer of delayed dmx signals */
-	TQueue<FDMXSignal> DelayedBuffer;
 };
