DMX - Fix CIS GoldenPath Win64 Development BootTest

diff --git a/Engine/Plugins/VirtualProduction/DMX/DMXProtocol/Source/DMXProtocol/Public/DMXProtocolObjectVersion.cpp b/Engine/Plugins/VirtualProduction/DMX/DMXProtocol/Source/DMXProtocol/Public/DMXProtocolObjectVersion.cpp
--- a/Engine/Plugins/VirtualProduction/DMX/DMXProtocol/Source/DMXProtocol/Public/DMXProtocolObjectVersion.cpp
+++ b/Engine/Plugins/VirtualProduction/DMX/DMXProtocol/Source/DMXProtocol/Public/DMXProtocolObjectVersion.cpp
@@ -7,4 +7,4 @@
 const FGuid FDMXProtocolObjectVersion::GUID(0x94662356, 0x4DAE92AC, 0xABF8BE98, 0x19FE8C2F);
 
 // Register the custom version with core
-FCustomVersionRegistration GRegisterDMXRuntimeCustomVersion(FDMXProtocolObjectVersion::GUID, FDMXProtocolObjectVersion::LatestVersion, TEXT("DMXProtocolObjectVersion"));
+FCustomVersionRegistration GRegisterDMXProtocolCustomVersion(FDMXProtocolObjectVersion::GUID, FDMXProtocolObjectVersion::LatestVersion, TEXT("DMXProtocolObjectVersion"));
