DMX - Fix delay, fix a crash when changing ports while a Unicast IP Address is empty

diff --git a/Engine/Plugins/VirtualProduction/DMX/DMXProtocol/Source/DMXProtocol/Private/DMXProtocolUtils.cpp b/Engine/Plugins/VirtualProduction/DMX/DMXProtocol/Source/DMXProtocol/Private/DMXProtocolUtils.cpp
--- a/Engine/Plugins/VirtualProduction/DMX/DMXProtocol/Source/DMXProtocol/Private/DMXProtocolUtils.cpp
+++ b/Engine/Plugins/VirtualProduction/DMX/DMXProtocol/Source/DMXProtocol/Private/DMXProtocolUtils.cpp
@@ -5,6 +5,7 @@
 #include "IPAddress.h"
 #include "SocketSubsystem.h" 
 
+
 #define LOCTEXT_NAMESPACE "FDMXProtocolUtils"
 
 TArray<TSharedPtr<FString>> FDMXProtocolUtils::GetLocalNetworkInterfaceCardIPs()
@@ -31,6 +32,27 @@ TArray<TSharedPtr<FString>> FDMXProtocolUtils::GetLocalNetworkInterfaceCardIPs()
 	return LocalNetworkInterfaceCardIPs;
 }
 
+TSharedPtr<FInternetAddr> FDMXProtocolUtils::CreateInternetAddr(const FString& IPAddress, int32 Port)
+{
+	if (IPAddress.IsEmpty())
+	{
+		return nullptr;
+	}
+
+	ISocketSubsystem* SocketSubsystem = ISocketSubsystem::Get(PLATFORM_SOCKETSUBSYSTEM);
+	TSharedPtr<FInternetAddr> InternetAddr = SocketSubsystem->CreateInternetAddr();
+
+	bool bIsValidIP = false;
+	InternetAddr->SetIp(*IPAddress, bIsValidIP);
+	if (!bIsValidIP)
+	{
+		return nullptr;
+	}
+
+	InternetAddr->SetPort(Port);
+	return InternetAddr;
+}
+
 FString FDMXProtocolUtils::GenerateUniqueNameFromExisting(const TSet<FString>& InExistingNames, const FString& InBaseName)
 {
 	if (!InBaseName.IsEmpty() && !InExistingNames.Contains(InBaseName))
diff --git a/Engine/Plugins/VirtualProduction/DMX/DMXProtocol/Source/DMXProtocol/Private/IO/DMXOutputPort.cpp b/Engine/Plugins/VirtualProduction/DMX/DMXProtocol/Source/DMXProtocol/Private/IO/DMXOutputPort.cpp
--- a/Engine/Plugins/VirtualProduction/DMX/DMXProtocol/Source/DMXProtocol/Private/IO/DMXOutputPort.cpp
+++ b/Engine/Plugins/VirtualProduction/DMX/DMXProtocol/Source/DMXProtocol/Private/IO/DMXOutputPort.cpp
@@ -14,50 +14,6 @@
 
 #define LOCTEXT_NAMESPACE "DMXOutputPort"
 
-
-
-FDMXSignalSharedPtr FDMXThreadSafeUniverseToSignalMap::GetSignal(int32 UniverseID) const
-{
-	const FDMXSignalSharedPtr* SignalPtr = UniverseToSignalMap.Find(UniverseID);
-	
-	return SignalPtr ? *SignalPtr : nullptr;
-}
-
-FDMXSignalSharedRef FDMXThreadSafeUniverseToSignalMap::GetOrCreateSignal(int32 UniverseID)
-{
-	FDMXSignalSharedPtr& Signal = UniverseToSignalMap.FindOrAdd(UniverseID);
-	if (!Signal.IsValid())
-	{
-		Signal = MakeShared<FDMXSignal, ESPMode::ThreadSafe>();
-	}
-
-	return Signal.ToSharedRef();
-}
-
-void FDMXThreadSafeUniverseToSignalMap::AddSignal(const FDMXSignalSharedRef& Signal)
-{
-	FScopeLock Lock(&CriticalSection);
-
-	UniverseToSignalMap.Add(Signal->ExternUniverseID, Signal);
-}
-
-void FDMXThreadSafeUniverseToSignalMap::ForEachSignal(TUniverseSignalPredicate Predicate)
-{
-	FScopeLock Lock(&CriticalSection);
-
-	for (TTuple<int32, FDMXSignalSharedPtr>& UniverseToSignalPair : UniverseToSignalMap)
-	{
-		Predicate(UniverseToSignalPair.Key, UniverseToSignalPair.Value);
-	}
-}
-
-void FDMXThreadSafeUniverseToSignalMap::Reset()
-{
-	FScopeLock Lock(&CriticalSection);
-
-	UniverseToSignalMap.Reset();
-}
-
 FDMXOutputPortSharedRef FDMXOutputPort::CreateFromConfig(FDMXOutputPortConfig& OutputPortConfig)
 {
 	// Port Configs are expected to have a valid guid always
@@ -188,7 +144,10 @@ void FDMXOutputPort::ClearBuffers()
 		RawListener->ClearBuffer();
 	}
 
-	LatestDMXSignals.Reset();
+	FScopeLock LockClearBuffers(&ClearBuffersCriticalSection);
+
+	SignalFragments.Empty();
+	ExternUniverseToLatestSignalMap.Reset();
 }
 
 bool FDMXOutputPort::IsLoopbackToEngine() const
@@ -219,9 +178,11 @@ bool FDMXOutputPort::GameThreadGetDMXSignal(int32 LocalUniverseID, FDMXSignalSha
 	{
 		int32 ExternUniverseID = ConvertLocalToExternUniverseID(LocalUniverseID);
 
-		OutDMXSignal = LatestDMXSignals.GetSignal(ExternUniverseID);
-		if (OutDMXSignal.IsValid())
+		FScopeLock LockExternUniverseToLatestSignalMap(&AccessExternUniverseToLatestSignalMapCriticalSection);
+		const FDMXSignalSharedPtr* SignalPtr = ExternUniverseToLatestSignalMap.Find(ExternUniverseID);
+		if (SignalPtr)
 		{
+			OutDMXSignal = *SignalPtr;
 			return true;
 		}
 	}
@@ -239,9 +200,11 @@ bool FDMXOutputPort::GameThreadGetDMXSignalFromRemoteUniverse(FDMXSignalSharedPt
 	const bool bNeedsLoopbackToEngine = CommunicationDeterminator.NeedsLoopbackToEngine();
 	if (bNeedsLoopbackToEngine || bEvenIfNotLoopbackToEngine)
 	{
-		OutDMXSignal = LatestDMXSignals.GetSignal(RemoteUniverseID);
-		if (OutDMXSignal.IsValid())
+		FScopeLock LockExternUniverseToLatestSignalMap(&AccessExternUniverseToLatestSignalMapCriticalSection);
+		const FDMXSignalSharedPtr* SignalPtr = ExternUniverseToLatestSignalMap.Find(RemoteUniverseID);
+		if (SignalPtr)
 		{
+			OutDMXSignal = *SignalPtr;
 			return true;
 		}
 	}
@@ -282,6 +245,10 @@ void FDMXOutputPort::RemoveRawListener(TSharedRef<FDMXRawListener> InRawListener
 
 void FDMXOutputPort::SendDMX(int32 LocalUniverseID, const TMap<int32, uint8>& ChannelToValueMap)
 {
+#if UE_BUILD_DEBUG
+	check(IsInGameThread());
+#endif // UE_BUILD_DEBUG
+
 	if (IsLocalUniverseInPortRange(LocalUniverseID))
 	{
 		const bool bNeedsSendDMX = CommunicationDeterminator.NeedsSendDMX();
@@ -292,25 +259,10 @@ void FDMXOutputPort::SendDMX(int32 LocalUniverseID, const TMap<int32, uint8>& Ch
 		{
 			const int32 ExternUniverseID = ConvertLocalToExternUniverseID(LocalUniverseID);
 
-			FDMXSignalSharedRef Signal = LatestDMXSignals.GetOrCreateSignal(ExternUniverseID);
+			const double SendTime = FPlatformTime::Seconds() + DelaySeconds;
+			const TSharedPtr<FDMXSignalFragment> Fragment = MakeShared<FDMXSignalFragment>(ExternUniverseID, ChannelToValueMap, SendTime);
 
-			// Write the fragment & meta data 
-			for (const TTuple<int32, uint8>& ChannelValueKvp : ChannelToValueMap)
-			{
-				int32 ChannelIndex = ChannelValueKvp.Key - 1;
-
-				// Filter invalid indicies so we can send bp calls here without testing them first.
-				if (Signal->ChannelData.IsValidIndex(ChannelIndex))
-				{
-					Signal->ChannelData[ChannelIndex] = ChannelValueKvp.Value;
-				}
-			}	
-
-			Signal->ExternUniverseID = ExternUniverseID;
-			Signal->Timestamp = FPlatformTime::Seconds() + DelaySeconds;
-			Signal->Priority = Priority;
-
-			NewDMXSignalsToSend.Enqueue(Signal);
+			SignalFragments.Enqueue(Fragment);
 		}
 	}
 }
@@ -318,6 +270,10 @@ void FDMXOutputPort::SendDMX(int32 LocalUniverseID, const TMap<int32, uint8>& Ch
 void FDMXOutputPort::SendDMXToRemoteUniverse(const TMap<int32, uint8>& ChannelToValueMap, int32 RemoteUniverse)
 {
 	// DEPRECATED 4.27
+#if UE_BUILD_DEBUG
+	check(IsInGameThread());
+#endif // UE_BUILD_DEBUG
+
 	if (IsExternUniverseInPortRange(RemoteUniverse))
 	{
 		const bool bNeedsSendDMX = CommunicationDeterminator.NeedsSendDMX();
@@ -326,24 +282,10 @@ void FDMXOutputPort::SendDMXToRemoteUniverse(const TMap<int32, uint8>& ChannelTo
 		// Update the buffer for loopback if dmx needs be sent and/or looped back
 		if (bNeedsSendDMX || bNeedsLoopbackToEngine)
 		{
-			FDMXSignalSharedRef Signal = LatestDMXSignals.GetOrCreateSignal(RemoteUniverse);
-
-			// Write the fragment & meta data 
-			for (const TTuple<int32, uint8>& ChannelValueKvp : ChannelToValueMap)
-			{
-				int32 ChannelIndex = ChannelValueKvp.Key - 1;
+			const double SendTime = FPlatformTime::Seconds() + DelaySeconds;
+			const TSharedPtr<FDMXSignalFragment> Fragment = MakeShared<FDMXSignalFragment>(RemoteUniverse, ChannelToValueMap, SendTime);
 
-				if (Signal->ChannelData.IsValidIndex(ChannelIndex))
-				{
-					Signal->ChannelData[ChannelIndex] = ChannelValueKvp.Value;
-				}
-			}
-
-			Signal->ExternUniverseID = RemoteUniverse;
-			Signal->Timestamp = FPlatformTime::Seconds() + DelaySeconds;
-			Signal->Priority = Priority;
-
-			NewDMXSignalsToSend.Enqueue(Signal);
+			SignalFragments.Enqueue(Fragment);
 		}
 	}
 }
@@ -471,19 +413,48 @@ void FDMXOutputPort::ProcessSendDMX()
 	// Delay signals
 	const double Now = FPlatformTime::Seconds();
 
-	// Acquire new DMX signals
-	for (;;)
+	FScopeLock LockClearBuffers(&ClearBuffersCriticalSection);
+	
+	// Write dmx fragments
 	{
-		FDMXSignalSharedPtr OldestDMXSignal;
-		if (NewDMXSignalsToSend.Peek(OldestDMXSignal))
+		for (;;)
 		{
-			if (OldestDMXSignal->Timestamp <= Now)
+			TSharedPtr<FDMXSignalFragment, ESPMode::ThreadSafe> OldestFragment;
+			if (SignalFragments.Peek(OldestFragment))
 			{
-				LatestDMXSignals.AddSignal(OldestDMXSignal.ToSharedRef());
+				if (OldestFragment->SendTime <= Now)
+				{
+					FDMXSignalSharedPtr* SignalPtr = ExternUniverseToLatestSignalMap.Find(OldestFragment->ExternUniverseID);
+					if (!SignalPtr)
+					{
+						FScopeLock LockExternUniverseToLatestSignalMap(&AccessExternUniverseToLatestSignalMapCriticalSection);
+						SignalPtr = &ExternUniverseToLatestSignalMap.Add(OldestFragment->ExternUniverseID, MakeShared<FDMXSignal, ESPMode::ThreadSafe>());
+					}
+					checkf(SignalPtr, TEXT("Cannot find or add a DMX signal to the DMX Output Port buffer"));
 
-				NewDMXSignalsToSend.Pop();
+					// Write the fragment & meta data 
+					FDMXSignalSharedPtr& Signal = *SignalPtr;
+					for (const TTuple<int32, uint8>& ChannelValueKvp : OldestFragment->ChannelToValueMap)
+					{
+						int32 ChannelIndex = ChannelValueKvp.Key - 1;
+						// Filter invalid indicies so we can send bp calls here without testing them first.
+						if (Signal->ChannelData.IsValidIndex(ChannelIndex))
+						{
+							Signal->ChannelData[ChannelIndex] = ChannelValueKvp.Value;
+						}
+					}
+
+					Signal->ExternUniverseID = OldestFragment->ExternUniverseID;
+					Signal->Timestamp = Now;
+					Signal->Priority = Priority;
 
-				continue;
+					// Drop the written fragment
+					SignalFragments.Pop();
+
+					continue;
+				}
+
+				break;
 			}
 
 			break;
@@ -492,30 +463,30 @@ void FDMXOutputPort::ProcessSendDMX()
 
 	// Send new and alive DMX Signals
 	const bool bNeedsSendDMX = CommunicationDeterminator.NeedsSendDMX();
-	LatestDMXSignals.ForEachSignal([Now, bNeedsSendDMX, this](int32 UniverseID, const FDMXSignalSharedPtr& Signal)
+	for (const TTuple<int32, FDMXSignalSharedPtr>& UniverseToSignalPair : ExternUniverseToLatestSignalMap)
+	{
+		if (UniverseToSignalPair.Value->Timestamp <= Now)
 		{
-			if (Signal->Timestamp <= Now)
-			{
-				// Keeping the signal alive here:
-				// Increment the timestamp by one second so the signal will be sent in one second anew.
-				Signal->Timestamp = Now + 1000.0;
+			// Keeping the signal alive here:
+			// Increment the timestamp by one second so the signal will be sent anew in one second.
+			UniverseToSignalPair.Value->Timestamp = Now + 1.0;
 
-				// Send via the protocol's sender
-				if (bNeedsSendDMX)
+			// Send via the protocol's sender
+			if (bNeedsSendDMX)
+			{
+				for (const TSharedPtr<IDMXSender>& DMXSender : DMXSenderArray)
 				{
-					for (const TSharedPtr<IDMXSender>& DMXSender : DMXSenderArray)
-					{
-						DMXSender->SendDMXSignal(Signal.ToSharedRef());
-					}
+					DMXSender->SendDMXSignal(UniverseToSignalPair.Value.ToSharedRef());
 				}
+			}
 
-				// Loopback to Listeners
-				for (const TSharedRef<FDMXRawListener>& RawListener : RawListeners)
-				{
-					RawListener->EnqueueSignal(this, Signal.ToSharedRef());
-				}
+			// Loopback to Listeners
+			for (const TSharedRef<FDMXRawListener>& RawListener : RawListeners)
+			{
+				RawListener->EnqueueSignal(this, UniverseToSignalPair.Value.ToSharedRef());
 			}
-		});
+		}
+	}
 }
 
 #undef LOCTEXT_NAMESPACE
diff --git a/Engine/Plugins/VirtualProduction/DMX/DMXProtocol/Source/DMXProtocol/Public/DMXProtocolUtils.h b/Engine/Plugins/VirtualProduction/DMX/DMXProtocol/Source/DMXProtocol/Public/DMXProtocolUtils.h
--- a/Engine/Plugins/VirtualProduction/DMX/DMXProtocol/Source/DMXProtocol/Public/DMXProtocolUtils.h
+++ b/Engine/Plugins/VirtualProduction/DMX/DMXProtocol/Source/DMXProtocol/Public/DMXProtocolUtils.h
@@ -4,6 +4,8 @@
 
 #include "CoreMinimal.h"
 
+class FInternetAddr;
+
 
 class DMXPROTOCOL_API FDMXProtocolUtils
 {
@@ -11,6 +13,16 @@ public:
 	/** Returns an array of local network interfaces */
 	static TArray<TSharedPtr<FString>> GetLocalNetworkInterfaceCardIPs();
 
+	/** 
+	 * Creates an Internet Address 
+	 * 
+	 * @param IPAddress	The IP Address to use
+	 * @param Port		The Port to use
+	 * 
+	 * @return			The Internet Address or nullptr if IPAddress and Port don't yield a valid Internet Address
+	 */
+	static TSharedPtr<FInternetAddr> CreateInternetAddr(const FString& IPAddress, int32 Port);
+
 	/**
 	 * Generates a unique name given a base one and a list of existing ones, by appending an index to
 	 * existing names. If InBaseName is an empty String, it returns "Default name".
diff --git a/Engine/Plugins/VirtualProduction/DMX/DMXProtocol/Source/DMXProtocol/Public/IO/DMXOutputPort.h b/Engine/Plugins/VirtualProduction/DMX/DMXProtocol/Source/DMXProtocol/Public/IO/DMXOutputPort.h
--- a/Engine/Plugins/VirtualProduction/DMX/DMXProtocol/Source/DMXProtocol/Public/IO/DMXOutputPort.h
+++ b/Engine/Plugins/VirtualProduction/DMX/DMXProtocol/Source/DMXProtocol/Public/IO/DMXOutputPort.h
@@ -60,35 +60,29 @@ private:
 };
 
 
-/** Helper to access a Map of Signals per Universe in a thread-safe way */
-class FDMXThreadSafeUniverseToSignalMap
+/** Structs that holds the fragmented values to be sent along with a timestamp when to send it */
+struct DMXPROTOCOL_API FDMXSignalFragment
+	: TSharedFromThis<FDMXSignalFragment, ESPMode::ThreadSafe>
 {
-public:
-	using TUniverseSignalPredicate = TFunctionRef<void(int32, const FDMXSignalSharedPtr&)>;
-	
-	/** Gets a signal, returns an invalid shared pointer if there is no signal at given Universe ID */
-	FDMXSignalSharedPtr GetSignal(int32 UniverseID) const;
-
-	/** Gets or creates a signal */
-	FDMXSignalSharedRef GetOrCreateSignal(int32 UniverseID);
+	FDMXSignalFragment() = delete;
 
-	/** Adds a signal to the map */
-	void AddSignal(const FDMXSignalSharedRef& Signal);
-
-	/** Loops through all signals by given Predicate  */
-	void ForEachSignal(TUniverseSignalPredicate Predicate);
+	FDMXSignalFragment(int32 InExternUniverseID, const TMap<int32, uint8>& InChannelToValueMap, double InSendTime)
+		: ExternUniverseID(InExternUniverseID)
+		, ChannelToValueMap(InChannelToValueMap)
+		, SendTime(InSendTime)
+	{}
 
-	/** Resets the Map */
-	void Reset();
+	/** The universe the fragment needs to be written to */
+	int32 ExternUniverseID;
 
-private:
-	/** Buffer of all the latest DMX Signals that were sent */
-	TMap<int32, FDMXSignalSharedPtr> UniverseToSignalMap;
+	/** The map of channels and values to send */
+	TMap<int32, uint8> ChannelToValueMap;
 
-	/** Critical sequestion required to be used when the Map is accessed */
-	FCriticalSection CriticalSection;
+	/** The time when the Fragment needs be sent */
+	double SendTime;
 };
 
+
 /**
  * Higher level abstraction of a DMX input hiding networking specific and protocol specific complexity.
  *
@@ -194,10 +188,10 @@ private:
 	TArray<TSharedPtr<IDMXSender>> DMXSenderArray;
 
 	/** Buffer of the signals that are to be sent in the next frame */
-	TQueue<FDMXSignalSharedPtr> NewDMXSignalsToSend;
+	TQueue<TSharedPtr<FDMXSignalFragment, ESPMode::ThreadSafe>> SignalFragments;
 
 	/** Map that holds the latest Signal per Universe */
-	FDMXThreadSafeUniverseToSignalMap LatestDMXSignals;
+	TMap<int32, FDMXSignalSharedPtr> ExternUniverseToLatestSignalMap;
 
 	/** The Destination Address to send to, can be irrelevant, e.g. for art-net broadcast */
 	TArray<FDMXOutputPortDestinationAddress> DestinationAddresses;
@@ -220,6 +214,12 @@ private:
 	/** Delay to apply on packets being sent */
 	double DelaySeconds = 0.0;
 
+	/** Critical section required to be used when the SignalFragments are accessed across threads */
+	FCriticalSection AccessExternUniverseToLatestSignalMapCriticalSection;
+
+	/** Critical section required to be used when clearing buffers */
+	FCriticalSection ClearBuffersCriticalSection;
+
 	/** Holds the thread object. */
 	FRunnableThread* Thread = nullptr;
 
diff --git a/Engine/Plugins/VirtualProduction/DMX/DMXProtocol/Source/DMXProtocolArtNet/Private/DMXProtocolArtNet.cpp b/Engine/Plugins/VirtualProduction/DMX/DMXProtocol/Source/DMXProtocolArtNet/Private/DMXProtocolArtNet.cpp
--- a/Engine/Plugins/VirtualProduction/DMX/DMXProtocol/Source/DMXProtocolArtNet/Private/DMXProtocolArtNet.cpp
+++ b/Engine/Plugins/VirtualProduction/DMX/DMXProtocol/Source/DMXProtocolArtNet/Private/DMXProtocolArtNet.cpp
@@ -146,7 +146,6 @@ void FDMXProtocolArtNet::UnregisterInputPort(const FDMXInputPortSharedRef& Input
 	}
 }
 
-PRAGMA_DISABLE_OPTIMIZATION
 TArray<TSharedPtr<IDMXSender>> FDMXProtocolArtNet::RegisterOutputPort(const FDMXOutputPortSharedRef& OutputPort)
 {
 	check(!OutputPort->IsRegistered());
@@ -209,7 +208,6 @@ TArray<TSharedPtr<IDMXSender>> FDMXProtocolArtNet::RegisterOutputPort(const FDMX
 
 	return NewSenders;
 }
-PRAGMA_ENABLE_OPTIMIZATION
 
 void FDMXProtocolArtNet::UnregisterOutputPort(const FDMXOutputPortSharedRef& OutputPort)
 {
diff --git a/Engine/Plugins/VirtualProduction/DMX/DMXProtocol/Source/DMXProtocolArtNet/Private/DMXProtocolArtNetReceiver.cpp b/Engine/Plugins/VirtualProduction/DMX/DMXProtocol/Source/DMXProtocolArtNet/Private/DMXProtocolArtNetReceiver.cpp
--- a/Engine/Plugins/VirtualProduction/DMX/DMXProtocol/Source/DMXProtocolArtNet/Private/DMXProtocolArtNetReceiver.cpp
+++ b/Engine/Plugins/VirtualProduction/DMX/DMXProtocol/Source/DMXProtocolArtNet/Private/DMXProtocolArtNetReceiver.cpp
@@ -5,6 +5,7 @@
 #include "DMXProtocolLog.h"
 #include "DMXProtocolArtNet.h"
 #include "DMXProtocolArtNetConstants.h"
+#include "DMXProtocolUtils.h"
 #include "DMXStats.h"
 #include "IO/DMXInputPort.h"
 #include "Packets/DMXProtocolArtNetPackets.h"
@@ -52,7 +53,7 @@ FDMXProtocolArtNetReceiver::~FDMXProtocolArtNetReceiver()
 
 TSharedPtr<FDMXProtocolArtNetReceiver> FDMXProtocolArtNetReceiver::TryCreate(const TSharedPtr<FDMXProtocolArtNet, ESPMode::ThreadSafe>& ArtNetProtocol, const FString& IPAddress)
 {
-	TSharedPtr<FInternetAddr> EndpointInternetAddr = CreateEndpointInternetAddr(IPAddress);
+	TSharedPtr<FInternetAddr> EndpointInternetAddr = FDMXProtocolUtils::CreateInternetAddr(IPAddress, ARTNET_PORT);
 	if (!EndpointInternetAddr.IsValid())
 	{
 		UE_LOG(LogDMXProtocol, Error, TEXT("Cannot create Art-Net receiver: Invalid IP address: %s"), *IPAddress);
@@ -79,7 +80,7 @@ TSharedPtr<FDMXProtocolArtNetReceiver> FDMXProtocolArtNetReceiver::TryCreate(con
 
 bool FDMXProtocolArtNetReceiver::EqualsEndpoint(const FString& IPAddress) const
 {
-	TSharedPtr<FInternetAddr> OtherEndpointInternetAddr = CreateEndpointInternetAddr(IPAddress);
+	TSharedPtr<FInternetAddr> OtherEndpointInternetAddr = FDMXProtocolUtils::CreateInternetAddr(IPAddress, ARTNET_PORT);
 	if (OtherEndpointInternetAddr.IsValid() && OtherEndpointInternetAddr->CompareEndpoints(*EndpointInternetAddr))
 	{
 		return true;
@@ -100,24 +101,6 @@ void FDMXProtocolArtNetReceiver::UnassignInputPort(const TSharedPtr<FDMXInputPor
 	AssignedInputPorts.Remove(InputPort);
 }
 
-TSharedPtr<FInternetAddr> FDMXProtocolArtNetReceiver::CreateEndpointInternetAddr(const FString& IPAddress)
-{
-	ISocketSubsystem* SocketSubsystem = ISocketSubsystem::Get(PLATFORM_SOCKETSUBSYSTEM);
-
-	TSharedPtr<FInternetAddr> InternetAddr = SocketSubsystem->CreateInternetAddr();
-
-	bool bIsValidIP = false;
-	InternetAddr->SetIp(*IPAddress, bIsValidIP);
-	if (!bIsValidIP)
-	{
-		return nullptr;
-	}
-
-	InternetAddr->SetPort(ARTNET_PORT);
-	return InternetAddr;
-}
-
-
 bool FDMXProtocolArtNetReceiver::Init()
 {
 	return true;
@@ -162,7 +145,7 @@ void FDMXProtocolArtNetReceiver::Update(const FTimespan& SocketWaitTime)
 	}
 
 	ISocketSubsystem* SocketSubsystem = ISocketSubsystem::Get(PLATFORM_SOCKETSUBSYSTEM);
-	TSharedPtr<FInternetAddr> Sender = SocketSubsystem->CreateInternetAddr();
+	TSharedPtr<FInternetAddr> InternetAddr = SocketSubsystem->CreateInternetAddr();
 
 	uint32 Size = 0;
 	int32 Read = 0;
@@ -174,7 +157,7 @@ void FDMXProtocolArtNetReceiver::Update(const FTimespan& SocketWaitTime)
 		constexpr uint32 ArtNetMaxReaderSize = 1024u;
 		Reader->SetNumUninitialized(FMath::Min(Size, ArtNetMaxReaderSize));
 		
-        if (Socket->RecvFrom(Reader->GetData(), Reader->Num(), Read, *Sender))
+        if (Socket->RecvFrom(Reader->GetData(), Reader->Num(), Read, *InternetAddr))
 		{
             Reader->RemoveAt(Read, Reader->Num() - Read, false);
 			
diff --git a/Engine/Plugins/VirtualProduction/DMX/DMXProtocol/Source/DMXProtocolArtNet/Private/DMXProtocolArtNetReceiver.h b/Engine/Plugins/VirtualProduction/DMX/DMXProtocol/Source/DMXProtocolArtNet/Private/DMXProtocolArtNetReceiver.h
--- a/Engine/Plugins/VirtualProduction/DMX/DMXProtocol/Source/DMXProtocolArtNet/Private/DMXProtocolArtNetReceiver.h
+++ b/Engine/Plugins/VirtualProduction/DMX/DMXProtocol/Source/DMXProtocolArtNet/Private/DMXProtocolArtNetReceiver.h
@@ -58,9 +58,6 @@ public:
 	FORCEINLINE const TSet<TSharedPtr<FDMXInputPort, ESPMode::ThreadSafe>>& GetAssignedInputPorts() const { return AssignedInputPorts; }
 
 private:
-	/** Helper to Create an internet address from an IP address string. Returns the InternetAddr or nullptr if unsuccessful */
-	static TSharedPtr<FInternetAddr> CreateEndpointInternetAddr(const FString& IPAddress);
-
 	/** The input ports the receiver uses */
 	TSet<TSharedPtr<FDMXInputPort, ESPMode::ThreadSafe>> AssignedInputPorts;
 
diff --git a/Engine/Plugins/VirtualProduction/DMX/DMXProtocol/Source/DMXProtocolArtNet/Private/DMXProtocolArtNetSender.cpp b/Engine/Plugins/VirtualProduction/DMX/DMXProtocol/Source/DMXProtocolArtNet/Private/DMXProtocolArtNetSender.cpp
--- a/Engine/Plugins/VirtualProduction/DMX/DMXProtocol/Source/DMXProtocolArtNet/Private/DMXProtocolArtNetSender.cpp
+++ b/Engine/Plugins/VirtualProduction/DMX/DMXProtocol/Source/DMXProtocolArtNet/Private/DMXProtocolArtNetSender.cpp
@@ -6,6 +6,7 @@
 #include "DMXProtocolConstants.h"
 #include "DMXProtocolLog.h"
 #include "DMXProtocolSettings.h"
+#include "DMXProtocolUtils.h"
 #include "DMXStats.h"
 #include "Packets/DMXProtocolArtNetPackets.h"
 
@@ -50,7 +51,7 @@ TSharedPtr<FDMXProtocolArtNetSender> FDMXProtocolArtNetSender::TryCreateUnicastS
 	const FString& InUnicastIP)
 {
 	// Try to create a socket
-	TSharedPtr<FInternetAddr> NewNetworkInterfaceInternetAddr = CreateInternetAddr(InNetworkInterfaceIP, ARTNET_SENDER_PORT);
+	TSharedPtr<FInternetAddr> NewNetworkInterfaceInternetAddr = FDMXProtocolUtils::CreateInternetAddr(InNetworkInterfaceIP, ARTNET_SENDER_PORT);
 	if (!NewNetworkInterfaceInternetAddr.IsValid())
 	{
 		UE_LOG(LogDMXProtocol, Error, TEXT("Cannot create Art-Net sender: Invalid IP address: %s"), *InNetworkInterfaceIP);
@@ -72,7 +73,7 @@ TSharedPtr<FDMXProtocolArtNetSender> FDMXProtocolArtNetSender::TryCreateUnicastS
 	}
 
 	// Try create the unicast internet addr
-	TSharedPtr<FInternetAddr> NewUnicastInternetAddr = CreateInternetAddr(InUnicastIP, ARTNET_PORT);
+	TSharedPtr<FInternetAddr> NewUnicastInternetAddr = FDMXProtocolUtils::CreateInternetAddr(InUnicastIP, ARTNET_PORT);
 	if (!NewUnicastInternetAddr.IsValid())
 	{
 		UE_LOG(LogDMXProtocol, Error, TEXT("Invalid Unicast IP %s for DMX Port. Please update your Output Port in Project Settings -> Plugins -> DMX Plugin"), *InUnicastIP);
@@ -90,7 +91,7 @@ TSharedPtr<FDMXProtocolArtNetSender> FDMXProtocolArtNetSender::TryCreateBroadcas
 	const FString& InNetworkInterfaceIP)
 {
 	// Try to create a socket
-	TSharedPtr<FInternetAddr> NewNetworkInterfaceInternetAddr = CreateInternetAddr(InNetworkInterfaceIP, ARTNET_SENDER_PORT);
+	TSharedPtr<FInternetAddr> NewNetworkInterfaceInternetAddr = FDMXProtocolUtils::CreateInternetAddr(InNetworkInterfaceIP, ARTNET_SENDER_PORT);
 	if (!NewNetworkInterfaceInternetAddr.IsValid())
 	{
 		UE_LOG(LogDMXProtocol, Error, TEXT("Cannot create Art-Net sender: Invalid IP address: %s"), *InNetworkInterfaceIP);
@@ -122,10 +123,10 @@ TSharedPtr<FDMXProtocolArtNetSender> FDMXProtocolArtNetSender::TryCreateBroadcas
 
 bool FDMXProtocolArtNetSender::EqualsEndpoint(const FString& NetworkInterfaceIP, const FString& DestinationIPAddress) const
 {
-	TSharedPtr<FInternetAddr> OtherNetworkInterfaceInternetAddr = CreateInternetAddr(NetworkInterfaceIP, ARTNET_SENDER_PORT);
+	TSharedPtr<FInternetAddr> OtherNetworkInterfaceInternetAddr = FDMXProtocolUtils::CreateInternetAddr(NetworkInterfaceIP, ARTNET_SENDER_PORT);
 	if (OtherNetworkInterfaceInternetAddr.IsValid() && OtherNetworkInterfaceInternetAddr->CompareEndpoints(*NetworkInterfaceInternetAddr))
 	{
-		TSharedPtr<FInternetAddr> OtherDestinationInternetAddr = CreateInternetAddr(DestinationIPAddress, ARTNET_PORT);
+		TSharedPtr<FInternetAddr> OtherDestinationInternetAddr = FDMXProtocolUtils::CreateInternetAddr(DestinationIPAddress, ARTNET_PORT);
 		if (OtherDestinationInternetAddr.IsValid() && OtherDestinationInternetAddr->CompareEndpoints(*DestinationInternetAddr))
 		{
 			return true;
@@ -198,23 +199,6 @@ void FDMXProtocolArtNetSender::SendDMXSignal(const FDMXSignalSharedRef& DMXSigna
 	}
 }
 
-TSharedPtr<FInternetAddr> FDMXProtocolArtNetSender::CreateInternetAddr(const FString& IPAddress, int32 Port)
-{
-	ISocketSubsystem* SocketSubsystem = ISocketSubsystem::Get(PLATFORM_SOCKETSUBSYSTEM);
-
-	TSharedPtr<FInternetAddr> InternetAddr = SocketSubsystem->CreateInternetAddr();
-
-	bool bIsValidIP = false;
-	InternetAddr->SetIp(*IPAddress, bIsValidIP);
-	if (!bIsValidIP)
-	{
-		return nullptr;
-	}
-
-	InternetAddr->SetPort(Port);
-	return InternetAddr;
-}
-
 TSharedRef<FInternetAddr> FDMXProtocolArtNetSender::CreateBroadcastInternetAddr(int32 Port)
 {
 	ISocketSubsystem* SocketSubsystem = ISocketSubsystem::Get(PLATFORM_SOCKETSUBSYSTEM);
diff --git a/Engine/Plugins/VirtualProduction/DMX/DMXProtocol/Source/DMXProtocolArtNet/Private/DMXProtocolArtNetSender.h b/Engine/Plugins/VirtualProduction/DMX/DMXProtocol/Source/DMXProtocolArtNet/Private/DMXProtocolArtNetSender.h
--- a/Engine/Plugins/VirtualProduction/DMX/DMXProtocol/Source/DMXProtocolArtNet/Private/DMXProtocolArtNetSender.h
+++ b/Engine/Plugins/VirtualProduction/DMX/DMXProtocol/Source/DMXProtocolArtNet/Private/DMXProtocolArtNetSender.h
@@ -78,9 +78,6 @@ public:
 	FORCEINLINE const TSet<TSharedPtr<FDMXOutputPort, ESPMode::ThreadSafe>>& GetAssignedOutputPorts() const { return AssignedOutputPorts; }
 
 private:
-	/** Helper to create an internet address from an IP address string. Returns the InternetAddr or nullptr if unsuccessful */
-	static TSharedPtr<FInternetAddr> CreateInternetAddr(const FString& IPAddress, int32 Port);
-
 	/** Helper to create a broadcast internet address. */
 	static TSharedRef<FInternetAddr> CreateBroadcastInternetAddr(int32 Port);
 
diff --git a/Engine/Plugins/VirtualProduction/DMX/DMXProtocol/Source/DMXProtocolSACN/Private/DMXProtocolSACNReceiver.cpp b/Engine/Plugins/VirtualProduction/DMX/DMXProtocol/Source/DMXProtocolSACN/Private/DMXProtocolSACNReceiver.cpp
--- a/Engine/Plugins/VirtualProduction/DMX/DMXProtocol/Source/DMXProtocolSACN/Private/DMXProtocolSACNReceiver.cpp
+++ b/Engine/Plugins/VirtualProduction/DMX/DMXProtocol/Source/DMXProtocolSACN/Private/DMXProtocolSACNReceiver.cpp
@@ -6,6 +6,7 @@
 #include "DMXProtocolSACN.h"
 #include "DMXProtocolSACNConstants.h"
 #include "DMXProtocolSACNUtils.h"
+#include "DMXProtocolUtils.h"
 #include "DMXStats.h"
 #include "IO/DMXInputPort.h"
 #include "Packets/DMXProtocolE131PDUPacket.h"
@@ -19,26 +20,6 @@
 DECLARE_DWORD_ACCUMULATOR_STAT(TEXT("sACN Packages Recieved Total"), STAT_SACNPackagesReceived, STATGROUP_DMX);
 
 
-namespace
-{
-	static TSharedPtr<FInternetAddr> CreateEndpointInternetAddr(const FString& IPAddress)
-	{
-		ISocketSubsystem* SocketSubsystem = ISocketSubsystem::Get(PLATFORM_SOCKETSUBSYSTEM);
-
-		TSharedPtr<FInternetAddr> InternetAddr = SocketSubsystem->CreateInternetAddr();
-
-		bool bIsValidIP = false;
-		InternetAddr->SetIp(*IPAddress, bIsValidIP);
-		if (!bIsValidIP)
-		{
-			return nullptr;
-		}
-
-		InternetAddr->SetPort(ACN_PORT);
-		return InternetAddr;
-	}
-}
-
 FDMXProtocolSACNReceiver::FDMXProtocolSACNReceiver(const TSharedPtr<FDMXProtocolSACN, ESPMode::ThreadSafe>& InSACNProtocol, FSocket& InSocket, TSharedRef<FInternetAddr> InEndpointInternetAddr)
 	: Protocol(InSACNProtocol)
 	, Socket(&InSocket)
@@ -73,7 +54,7 @@ FDMXProtocolSACNReceiver::~FDMXProtocolSACNReceiver()
 
 TSharedPtr<FDMXProtocolSACNReceiver> FDMXProtocolSACNReceiver::TryCreate(const TSharedPtr<FDMXProtocolSACN, ESPMode::ThreadSafe>& SACNProtocol, const FString& IPAddress)
 {
-	TSharedPtr<FInternetAddr> EndpointInternetAddr = CreateEndpointInternetAddr(IPAddress);
+	TSharedPtr<FInternetAddr> EndpointInternetAddr = FDMXProtocolUtils::CreateInternetAddr(IPAddress, ACN_PORT);
 	if (!EndpointInternetAddr.IsValid())
 	{
 		UE_LOG(LogDMXProtocol, Error, TEXT("Cannot create sACN receiver: Invalid IP address: %s"), *IPAddress);
@@ -112,7 +93,7 @@ TSharedPtr<FDMXProtocolSACNReceiver> FDMXProtocolSACNReceiver::TryCreate(const T
 
 bool FDMXProtocolSACNReceiver::EqualsEndpoint(const FString& IPAddress) const
 {
-	TSharedPtr<FInternetAddr> OtherEndpointInternetAddr = CreateEndpointInternetAddr(IPAddress);
+	TSharedPtr<FInternetAddr> OtherEndpointInternetAddr = FDMXProtocolUtils::CreateInternetAddr(IPAddress, ACN_PORT);
 	if (OtherEndpointInternetAddr.IsValid() && OtherEndpointInternetAddr->CompareEndpoints(*EndpointInternetAddr))
 	{
 		return true;
diff --git a/Engine/Plugins/VirtualProduction/DMX/DMXProtocol/Source/DMXProtocolSACN/Private/DMXProtocolSACNSender.cpp b/Engine/Plugins/VirtualProduction/DMX/DMXProtocol/Source/DMXProtocolSACN/Private/DMXProtocolSACNSender.cpp
--- a/Engine/Plugins/VirtualProduction/DMX/DMXProtocol/Source/DMXProtocolSACN/Private/DMXProtocolSACNSender.cpp
+++ b/Engine/Plugins/VirtualProduction/DMX/DMXProtocol/Source/DMXProtocolSACN/Private/DMXProtocolSACNSender.cpp
@@ -2,12 +2,13 @@
 
 #include "DMXProtocolSACNSender.h"
 
+#include "DMXProtocolConstants.h"
+#include "DMXProtocolLog.h"
 #include "DMXProtocolSACN.h"
 #include "DMXProtocolSACNReceiver.h"
 #include "DMXProtocolSACNUtils.h"
-#include "DMXProtocolConstants.h"
-#include "DMXProtocolLog.h"
 #include "DMXProtocolSettings.h"
+#include "DMXProtocolUtils.h"
 #include "DMXStats.h"
 #include "Packets/DMXProtocolE131PDUPacket.h"
 
@@ -22,27 +23,6 @@
 
 DECLARE_DWORD_ACCUMULATOR_STAT(TEXT("sACN Packages Sent Total"), STAT_SACNPackagesSent, STATGROUP_DMX);
 
-namespace
-{
-	/** Helper to create an internet address from an IP address string. Returns the InternetAddr or nullptr if unsuccessful */
-	TSharedPtr<FInternetAddr> CreateInternetAddr(const FString& IPAddress, int32 Port)
-	{
-		ISocketSubsystem* SocketSubsystem = ISocketSubsystem::Get(PLATFORM_SOCKETSUBSYSTEM);
-
-		TSharedPtr<FInternetAddr> InternetAddr = SocketSubsystem->CreateInternetAddr();
-
-		bool bIsValidIP = false;
-		InternetAddr->SetIp(*IPAddress, bIsValidIP);
-		if (!bIsValidIP)
-		{
-			return nullptr;
-		}
-
-		InternetAddr->SetPort(Port);
-		return InternetAddr;
-	}
-}
-
 
 FDMXProtocolSACNSender::FDMXProtocolSACNSender(const TSharedPtr<FDMXProtocolSACN, ESPMode::ThreadSafe>& InSACNProtocol, FSocket& InSocket, TSharedRef<FInternetAddr> InNetworkInterfaceInternetAddr, TSharedRef<FInternetAddr> InDestinationInternetAddr, const bool bInIsMulticast)
 	: Protocol(InSACNProtocol)
@@ -81,7 +61,7 @@ TSharedPtr<FDMXProtocolSACNSender> FDMXProtocolSACNSender::TryCreateUnicastSende
 	const FString& InUnicastIP)
 {
 	// Try to create a socket
-	TSharedPtr<FInternetAddr> NewNetworkInterfaceInternetAddr = CreateInternetAddr(InNetworkInterfaceIP, ACN_PORT);
+	TSharedPtr<FInternetAddr> NewNetworkInterfaceInternetAddr = FDMXProtocolUtils::CreateInternetAddr(InNetworkInterfaceIP, ACN_PORT);
 	if (!NewNetworkInterfaceInternetAddr.IsValid())
 	{
 		UE_LOG(LogDMXProtocol, Error, TEXT("Cannot create sACN sender: Invalid IP address: %s"), *InNetworkInterfaceIP);
@@ -104,7 +84,7 @@ TSharedPtr<FDMXProtocolSACNSender> FDMXProtocolSACNSender::TryCreateUnicastSende
 	}
 
 	// Try create the unicast internet addr
-	TSharedPtr<FInternetAddr> NewUnicastInternetAddr = CreateInternetAddr(InUnicastIP, ACN_PORT);
+	TSharedPtr<FInternetAddr> NewUnicastInternetAddr = FDMXProtocolUtils::CreateInternetAddr(InUnicastIP, ACN_PORT);
 	if (!NewUnicastInternetAddr.IsValid())
 	{
 		UE_LOG(LogDMXProtocol, Error, TEXT("Invalid Unicast IP %s for DMX Port. Please update your Output Port in Project Settings -> Plugins -> DMX Plugin"), *InUnicastIP);
@@ -122,7 +102,7 @@ TSharedPtr<FDMXProtocolSACNSender> FDMXProtocolSACNSender::TryCreateMulticastSen
 	const FString& InNetworkInterfaceIP)
 {
 	// Try to create a socket
-	TSharedPtr<FInternetAddr> NewNetworkInterfaceInternetAddr = CreateInternetAddr(InNetworkInterfaceIP, ACN_PORT);
+	TSharedPtr<FInternetAddr> NewNetworkInterfaceInternetAddr = FDMXProtocolUtils::CreateInternetAddr(InNetworkInterfaceIP, ACN_PORT);
 	if (!NewNetworkInterfaceInternetAddr.IsValid())
 	{
 		UE_LOG(LogDMXProtocol, Error, TEXT("Cannot create sACN sender: Invalid IP address: %s"), *InNetworkInterfaceIP);
@@ -154,10 +134,10 @@ TSharedPtr<FDMXProtocolSACNSender> FDMXProtocolSACNSender::TryCreateMulticastSen
 
 bool FDMXProtocolSACNSender::EqualsEndpoint(const FString& NetworkInterfaceIP, const FString& DestinationIPAddress) const
 {
-	TSharedPtr<FInternetAddr> OtherNetworkInterfaceInternetAddr = CreateInternetAddr(NetworkInterfaceIP, ACN_PORT);
+	TSharedPtr<FInternetAddr> OtherNetworkInterfaceInternetAddr = FDMXProtocolUtils::CreateInternetAddr(NetworkInterfaceIP, ACN_PORT);
 	if (OtherNetworkInterfaceInternetAddr.IsValid() && OtherNetworkInterfaceInternetAddr->CompareEndpoints(*NetworkInterfaceInternetAddr))
 	{
-		TSharedPtr<FInternetAddr> OtherDestinationInternetAddr = CreateInternetAddr(DestinationIPAddress, ACN_PORT);
+		TSharedPtr<FInternetAddr> OtherDestinationInternetAddr = FDMXProtocolUtils::CreateInternetAddr(DestinationIPAddress, ACN_PORT);
 		if (OtherDestinationInternetAddr.IsValid() && OtherDestinationInternetAddr->CompareEndpoints(*DestinationInternetAddr))
 		{
 			return true;
