DMX - Fix CIS GoldenPath Win64 Development BootTest

diff --git a/Engine/Plugins/VirtualProduction/DMX/DMXProtocol/Source/DMXProtocolArtNet/Private/DMXProtocolArtNetReceiver.cpp b/Engine/Plugins/VirtualProduction/DMX/DMXProtocol/Source/DMXProtocolArtNet/Private/DMXProtocolArtNetReceiver.cpp
--- a/Engine/Plugins/VirtualProduction/DMX/DMXProtocol/Source/DMXProtocolArtNet/Private/DMXProtocolArtNetReceiver.cpp
+++ b/Engine/Plugins/VirtualProduction/DMX/DMXProtocol/Source/DMXProtocolArtNet/Private/DMXProtocolArtNetReceiver.cpp
@@ -28,6 +28,9 @@ FDMXProtocolArtNetReceiver::FDMXProtocolArtNetReceiver(const TSharedPtr<FDMXProt
 {
 	check(Socket->GetSocketType() == SOCKTYPE_Datagram);
 
+	ISocketSubsystem* SocketSubsystem = ISocketSubsystem::Get(PLATFORM_SOCKETSUBSYSTEM);
+	ReceivedSenderInternetAddr = SocketSubsystem->CreateInternetAddr();
+
 	FString ReceiverThreadName = FString(TEXT("ArtNetReceiver_")) + InEndpointInternetAddr->ToString(true);
 	Thread = FRunnableThread::Create(this, *ReceiverThreadName, 0, TPri_TimeCritical, FPlatformAffinity::GetPoolThreadMask());
 
@@ -144,20 +147,17 @@ void FDMXProtocolArtNetReceiver::Update(const FTimespan& SocketWaitTime)
 		return;
 	}
 
-	ISocketSubsystem* SocketSubsystem = ISocketSubsystem::Get(PLATFORM_SOCKETSUBSYSTEM);
-	TSharedPtr<FInternetAddr> InternetAddr = SocketSubsystem->CreateInternetAddr();
-
 	uint32 Size = 0;
 	int32 Read = 0;
 	while (Socket->HasPendingData(Size))
 	{
-		TSharedRef<FArrayReader> Reader = MakeShared<FArrayReader>(true);
+		const TSharedRef<FArrayReader> Reader = MakeShared<FArrayReader>(true);
 
 		// Use an aligned size instead of ART_NET_PACKAGE_SIZE
 		constexpr uint32 ArtNetMaxReaderSize = 1024u;
 		Reader->SetNumUninitialized(FMath::Min(Size, ArtNetMaxReaderSize));
 		
-        if (Socket->RecvFrom(Reader->GetData(), Reader->Num(), Read, *InternetAddr))
+        if (Socket->RecvFrom(Reader->GetData(), Reader->Num(), Read, *ReceivedSenderInternetAddr))
 		{
             Reader->RemoveAt(Read, Reader->Num() - Read, false);
 			
diff --git a/Engine/Plugins/VirtualProduction/DMX/DMXProtocol/Source/DMXProtocolArtNet/Private/DMXProtocolArtNetReceiver.h b/Engine/Plugins/VirtualProduction/DMX/DMXProtocol/Source/DMXProtocolArtNet/Private/DMXProtocolArtNetReceiver.h
--- a/Engine/Plugins/VirtualProduction/DMX/DMXProtocol/Source/DMXProtocolArtNet/Private/DMXProtocolArtNetReceiver.h
+++ b/Engine/Plugins/VirtualProduction/DMX/DMXProtocol/Source/DMXProtocolArtNet/Private/DMXProtocolArtNetReceiver.h
@@ -106,9 +106,12 @@ private:
 	/** The network socket. */
 	FSocket* Socket = nullptr;
 
-	/** The endpoint internet addr */
+	/** The endpoint internet addr (usually the network interface card IP Address) */
 	TSharedPtr<FInternetAddr> EndpointInternetAddr;
 
+	/** The sender when packets were received */
+	TSharedPtr<FInternetAddr> ReceivedSenderInternetAddr;
+
 	/** Flag indicating that the thread is stopping. */
 	TAtomic<bool> bStopping;
 
diff --git a/Engine/Plugins/VirtualProduction/DMX/DMXProtocol/Source/DMXProtocolSACN/Private/DMXProtocolSACNReceiver.cpp b/Engine/Plugins/VirtualProduction/DMX/DMXProtocol/Source/DMXProtocolSACN/Private/DMXProtocolSACNReceiver.cpp
--- a/Engine/Plugins/VirtualProduction/DMX/DMXProtocol/Source/DMXProtocolSACN/Private/DMXProtocolSACNReceiver.cpp
+++ b/Engine/Plugins/VirtualProduction/DMX/DMXProtocol/Source/DMXProtocolSACN/Private/DMXProtocolSACNReceiver.cpp
@@ -29,6 +29,10 @@ FDMXProtocolSACNReceiver::FDMXProtocolSACNReceiver(const TSharedPtr<FDMXProtocol
 {
 	check(Socket->GetSocketType() == SOCKTYPE_Datagram);
 
+	ISocketSubsystem* SocketSubsystem = ISocketSubsystem::Get(PLATFORM_SOCKETSUBSYSTEM);
+	ReceivedSenderInternetAddr = SocketSubsystem->CreateInternetAddr();
+	ReceivedDestinationInternetAddr = SocketSubsystem->CreateInternetAddr();
+
 	FString ReceiverThreadName = FString(TEXT("SACNReceiver_")) + InEndpointInternetAddr->ToString(true);
 	Thread = FRunnableThread::Create(this, *ReceiverThreadName, 0, TPri_TimeCritical, FPlatformAffinity::GetPoolThreadMask());
 
@@ -211,10 +215,6 @@ void FDMXProtocolSACNReceiver::Update(const FTimespan& SocketWaitTime)
 		return;
 	}
 
-	ISocketSubsystem* SocketSubsystem = ISocketSubsystem::Get(PLATFORM_SOCKETSUBSYSTEM);
-	TSharedPtr<FInternetAddr> Sender = SocketSubsystem->CreateInternetAddr();
-	TSharedPtr<FInternetAddr> Destination = SocketSubsystem->CreateInternetAddr();
-
 	uint32 Size = 0;
 	int32 NumBytesRead = 0;
 	while (Socket->HasPendingData(Size))
@@ -225,11 +225,11 @@ void FDMXProtocolSACNReceiver::Update(const FTimespan& SocketWaitTime)
 		constexpr uint32 SACNMaxReaderSize = 1024u;
 		Reader->SetNumUninitialized(FMath::Min(Size, SACNMaxReaderSize));
 
-		if (Socket->RecvFromWithPktInfo(Reader->GetData(), Reader->Num(), NumBytesRead, *Sender, *Destination))
+		if (Socket->RecvFromWithPktInfo(Reader->GetData(), Reader->Num(), NumBytesRead, *ReceivedSenderInternetAddr, *ReceivedDestinationInternetAddr))
 		{
 			Reader->RemoveAt(NumBytesRead, Reader->Num() - NumBytesRead, false);
 			uint32 MulticastIp = 0;
-			Destination->GetIp(MulticastIp);
+			ReceivedDestinationInternetAddr->GetIp(MulticastIp);
 
 			if (const uint16* UniverseIDPtr = MulticastGroupAddrToUniverseIDMap.Find(MulticastIp))
 			{
diff --git a/Engine/Plugins/VirtualProduction/DMX/DMXProtocol/Source/DMXProtocolSACN/Private/DMXProtocolSACNReceiver.h b/Engine/Plugins/VirtualProduction/DMX/DMXProtocol/Source/DMXProtocolSACN/Private/DMXProtocolSACNReceiver.h
--- a/Engine/Plugins/VirtualProduction/DMX/DMXProtocol/Source/DMXProtocolSACN/Private/DMXProtocolSACNReceiver.h
+++ b/Engine/Plugins/VirtualProduction/DMX/DMXProtocol/Source/DMXProtocolSACN/Private/DMXProtocolSACNReceiver.h
@@ -103,9 +103,15 @@ private:
 	/** The network socket. */
 	FSocket* Socket = nullptr;
 
-	/** The endpoint internet addr */
+	/** The endpoint internet addr (usually the network interface card IP Address) */
 	TSharedPtr<FInternetAddr> EndpointInternetAddr;
 
+	/** The sender when packets are received */
+	TSharedPtr<FInternetAddr> ReceivedSenderInternetAddr;
+
+	/** The destination when packets are received */
+	TSharedPtr<FInternetAddr> ReceivedDestinationInternetAddr;
+
 	/** Flag indicating that the thread is stopping. */
 	TAtomic<bool> bStopping;
 
diff --git a/Engine/Plugins/VirtualProduction/RemoteControlProtocolDMX/Source/RemoteControlProtocolDMX/Private/RemoteControlProtocolDMXSettings.cpp b/Engine/Plugins/VirtualProduction/RemoteControlProtocolDMX/Source/RemoteControlProtocolDMX/Private/RemoteControlProtocolDMXSettings.cpp
--- a/Engine/Plugins/VirtualProduction/RemoteControlProtocolDMX/Source/RemoteControlProtocolDMX/Private/RemoteControlProtocolDMXSettings.cpp
+++ b/Engine/Plugins/VirtualProduction/RemoteControlProtocolDMX/Source/RemoteControlProtocolDMX/Private/RemoteControlProtocolDMXSettings.cpp
@@ -7,9 +7,9 @@
 
 FSimpleMulticastDelegate URemoteControlProtocolDMXSettings::OnRemoteControlProtocolDMXSettingsChangedDelegate;
 
-void URemoteControlProtocolDMXSettings::PostInitProperties()
+void URemoteControlProtocolDMXSettings::PostLoad()
 {
-	Super::PostInitProperties();
+	Super::PostLoad();
 
 	if (!DefaultInputPortId.IsValid() || !FDMXPortManager::Get().FindInputPortByGuid(DefaultInputPortId))
 	{
diff --git a/Engine/Plugins/VirtualProduction/RemoteControlProtocolDMX/Source/RemoteControlProtocolDMX/Public/RemoteControlProtocolDMXSettings.h b/Engine/Plugins/VirtualProduction/RemoteControlProtocolDMX/Source/RemoteControlProtocolDMX/Public/RemoteControlProtocolDMXSettings.h
--- a/Engine/Plugins/VirtualProduction/RemoteControlProtocolDMX/Source/RemoteControlProtocolDMX/Public/RemoteControlProtocolDMXSettings.h
+++ b/Engine/Plugins/VirtualProduction/RemoteControlProtocolDMX/Source/RemoteControlProtocolDMX/Public/RemoteControlProtocolDMXSettings.h
@@ -19,7 +19,7 @@ class REMOTECONTROLPROTOCOLDMX_API URemoteControlProtocolDMXSettings : public UO
 
 public:
 	//~ Begin UObject interface
-	virtual void PostInitProperties() override;
+	virtual void PostLoad() override;
 #if WITH_EDITOR
 	virtual void PostEditChangeProperty(struct FPropertyChangedEvent& PropertyChangedEvent) override;
 #endif
