DMX - Add console commands for various DMX port settings

diff --git a/Engine/Plugins/VirtualProduction/DMX/DMXProtocol/Source/DMXProtocol/Private/IO/DMXInputPort.cpp b/Engine/Plugins/VirtualProduction/DMX/DMXProtocol/Source/DMXProtocol/Private/IO/DMXInputPort.cpp
--- a/Engine/Plugins/VirtualProduction/DMX/DMXProtocol/Source/DMXProtocol/Private/IO/DMXInputPort.cpp
+++ b/Engine/Plugins/VirtualProduction/DMX/DMXProtocol/Source/DMXProtocol/Private/IO/DMXInputPort.cpp
@@ -9,12 +9,232 @@
 #include "IO/DMXPortManager.h"
 #include "IO/DMXRawListener.h"
 
+#include "HAL/IConsoleManager.h"
 
 DECLARE_CYCLE_STAT(TEXT("Input Port Tick"), STAT_DMXInputPortTick, STATGROUP_DMX);
 
 
 #define LOCTEXT_NAMESPACE "DMXInputPort"
 
+/** Helper to override a member variable of an Input Port */
+#define DMX_OVERRIDE_INPUTPORT_VAR(MemberName, PortName, Value) \
+{ \
+	const FDMXInputPortSharedRef* InputPortPtr = FDMXPortManager::Get().GetInputPorts().FindByPredicate([PortName](const FDMXInputPortSharedPtr& InputPort) \
+		{ \
+			return InputPort->GetPortName() == PortName; \
+		}); \
+	if (InputPortPtr) \
+	{ \
+		const FDMXInputPortConfig OldInputPortConfig = (*InputPortPtr)->MakeInputPortConfig(); \
+		FDMXInputPortConfigParams InputPortConfigParams = FDMXInputPortConfigParams(OldInputPortConfig); \
+		InputPortConfigParams.MemberName = Value; \
+		FDMXInputPortConfig NewInputPortConfig((*InputPortPtr)->GetPortGuid(), InputPortConfigParams); \
+		(*InputPortPtr)->UpdateFromConfig(NewInputPortConfig); \
+	} \
+}
 
+static FAutoConsoleCommand GDMXSetInputPortProtocolCommand(
+	TEXT("DMX.SetInputPortProtocol"),
+	TEXT("DMX.SetInputPortProtocol [PortName][ProtocolName]. Sets the protocol used by the input port. Example: DMX.SetInputPortProtocol MyInputPort sACN"),
+	FConsoleCommandWithArgsDelegate::CreateStatic(
+		[](const TArray<FString>& Args)
+		{
+			if (Args.Num() < 2)
+			{
+				return;
+			}
+
+			const FString& PortName = Args[0];
+			const FName ProtocolNameValue = FName(Args[1]);
+
+			if (IDMXProtocol::GetProtocolNames().Contains(ProtocolNameValue))
+			{
+				DMX_OVERRIDE_INPUTPORT_VAR(ProtocolName, PortName, ProtocolNameValue);
+			}
+		})
+);
+
+static FAutoConsoleCommand GDMXSetInputPortDeviceAddressCommand(
+	TEXT("DMX.SetInputPortDeviceAddress"),
+	TEXT("DMX.SetInputPortDeviceAddress [PortName][DeviceAddress]. Sets the Device Address of an input port, usually the network interface card IP address. Example: DMX.SetInputPortDeviceAddress MyInputPort 123.45.67.89"),
+	FConsoleCommandWithArgsDelegate::CreateStatic(
+		[](const TArray<FString>& Args)
+		{
+			if (Args.Num() < 2)
+			{
+				return;
+			}
+
+			const FString& PortName = Args[0];
+			const FString& DeviceAddressValue = Args[1];
+
+			DMX_OVERRIDE_INPUTPORT_VAR(DeviceAddress, PortName, DeviceAddressValue);
+		})
+);
+
+static FAutoConsoleCommand GDMXSetInputPortLocalUniverseStartCommand(
+	TEXT("DMX.SetInputPortLocalUniverseStart"),
+	TEXT("DMX.SetInputPortLocalUniverseStart [PortName][Universe]. Sets the local universe start of the input port. Example: DMX.SetInputPortLocalUniverseStart MyInputPort 5"),
+	FConsoleCommandWithArgsDelegate::CreateStatic(
+		[](const TArray<FString>& Args)
+		{
+			if (Args.Num() < 2)
+			{
+				return;
+			}
+
+			const FString& PortName = Args[0];
+			const FString& LocalUniverseStartValueString = Args[1];
+			int32 LocalUniverseStartValue;
+			if (LexTryParseString<int32>(LocalUniverseStartValue, *LocalUniverseStartValueString))
+			{
+				if (LocalUniverseStartValue >= 0 && LocalUniverseStartValue <= DMX_MAX_UNIVERSE)
+				{
+					DMX_OVERRIDE_INPUTPORT_VAR(LocalUniverseStart, PortName, LocalUniverseStartValue);
+				}
+			}
+		})
+);
+
+static FAutoConsoleCommand GDMXSetInputPortNumUniversesCommand(
+	TEXT("DMX.SetInputPortNumUniverses"),
+	TEXT("DMX.SetInputPortNumUniverses [PortName][Universe]. Sets the num universes of the input port. Example: DMX.SetInputPortNumUniverses MyInputPort 10"),
+	FConsoleCommandWithArgsDelegate::CreateStatic(
+		[](const TArray<FString>& Args)
+		{
+			if (Args.Num() < 2)
+			{
+				return;
+			}
+
+			const FString& PortName = Args[0];
+			const FString& NumUniversesValueString = Args[1];
+			int32 NumUniversesValue;
+			if (LexTryParseString<int32>(NumUniversesValue, *NumUniversesValueString))
+			{
+				if (NumUniversesValue >= 0 && NumUniversesValue <= DMX_MAX_UNIVERSE)
+				{
+					DMX_OVERRIDE_INPUTPORT_VAR(NumUniverses, PortName, NumUniversesValue);
+				}
+			}
+		})
+);
+
+static FAutoConsoleCommand GDMXSetInputPortExternUniverseStartCommand(
+	TEXT("DMX.SetInputPortExternUniverseStart"),
+	TEXT("DMX.SetInputPortExternUniverseStart [PortName][Universe]. Sets the extern universe start of the input port. Example: DMX.SetInputPortExternUniverseStart MyInputPort 7"),
+	FConsoleCommandWithArgsDelegate::CreateStatic(
+		[](const TArray<FString>& Args)
+		{
+			if (Args.Num() < 2)
+			{
+				return;
+			}
+
+			const FString& PortName = Args[0];
+			const FString& ExternUniverseStartValueString = Args[1];
+			int32 ExternUniverseStartValue;
+			if (LexTryParseString<int32>(ExternUniverseStartValue, *ExternUniverseStartValueString))
+			{
+				if (ExternUniverseStartValue >= 0 && ExternUniverseStartValue <= DMX_MAX_UNIVERSE)
+				{
+					DMX_OVERRIDE_INPUTPORT_VAR(ExternUniverseStart, PortName, ExternUniverseStartValue);
+				}
+			}
+		})
+);
+
+static FAutoConsoleCommand GDMXSetInputPortPriorityStrategyCommand(
+	TEXT("DMX.SetInputPortPriorityStrategy"),
+	TEXT("DMX.SetInputPortPriorityStrategy [PortName][PriorityStrategy (0 = None, 1 = Lowest, 2 = LowerThan, 3 = Equal, 4 = Higher Than, 5 = Highest)]. Example: DMX.SetInputPortPriorityStrategy MyInputPort 1"),
+	FConsoleCommandWithArgsDelegate::CreateStatic(
+		[](const TArray<FString>& Args)
+		{
+			if (Args.Num() < 2)
+			{
+				return;
+			}
+
+			const FString& PortName = Args[0];
+			const FString& PriorityStrategyValueString = Args[1];
+			uint8 PriorityStrategyValue;
+			if (LexTryParseString<uint8>(PriorityStrategyValue, *PriorityStrategyValueString))
+			{
+				const EDMXPortPriorityStrategy PriorityStrategyEnumValue = static_cast<EDMXPortPriorityStrategy>(PriorityStrategyValue);
+
+				if (PriorityStrategyEnumValue == EDMXPortPriorityStrategy::None ||
+					PriorityStrategyEnumValue == EDMXPortPriorityStrategy::Lowest ||
+					PriorityStrategyEnumValue == EDMXPortPriorityStrategy::LowerThan ||
+					PriorityStrategyEnumValue == EDMXPortPriorityStrategy::Equal ||
+					PriorityStrategyEnumValue == EDMXPortPriorityStrategy::HigherThan ||
+					PriorityStrategyEnumValue == EDMXPortPriorityStrategy::Highest)
+				{
+					DMX_OVERRIDE_INPUTPORT_VAR(PriorityStrategy, PortName, PriorityStrategyEnumValue);
+				}
+			}
+		})
+);
+
+static FAutoConsoleCommand GDMXSetInputPortPriorityCommand(
+	TEXT("DMX.SetInputPortPriority"),
+	TEXT("DMX.SetInputPortPriority [PortName][Universe]. Sets the priority of the input port. Example: DMX.SetInputPortPriority MyInputPort 100"),
+	FConsoleCommandWithArgsDelegate::CreateStatic(
+		[](const TArray<FString>& Args)
+		{
+			if (Args.Num() < 2)
+			{
+				return;
+			}
+
+			const FString& PortName = Args[0];
+			const FString& PriorityValueString = Args[1];
+			int32 PriorityValue;
+			if (LexTryParseString<int32>(PriorityValue, *PriorityValueString))
+			{
+				DMX_OVERRIDE_INPUTPORT_VAR(Priority, PortName, PriorityValue);
+			}
+		})
+);
+
+static FAutoConsoleCommand GDMXResetInputPortToProjectSettings(
+	TEXT("DMX.ResetInputPortToProjectSettings"),
+	TEXT("DMX.ResetInputPortToProjectSettings [PortName]. Resets the input port to how it is defined in project settings. Example: DMX.ResetInputPortToProjectSettings MyInputPort"),
+	FConsoleCommandWithArgsDelegate::CreateStatic(
+		[](const TArray<FString>& Args)
+		{
+			if (Args.Num() < 1)
+			{
+				return;
+			}
+
+			const FString& PortName = Args[0];
+			const FDMXInputPortSharedRef* PortPtr = FDMXPortManager::Get().GetInputPorts().FindByPredicate([PortName](const FDMXInputPortSharedPtr& InputPort)
+				{
+					return InputPort->GetPortName() == PortName;
+				});
+
+			if (PortPtr)
+			{
+				const UDMXProtocolSettings* ProtocolSettings = GetDefault<UDMXProtocolSettings>();
+				if (ProtocolSettings)
+				{
+					const FDMXInputPortConfig* PortConfigPtr = ProtocolSettings->InputPortConfigs.FindByPredicate([PortPtr](const FDMXInputPortConfig& InputPortConfig)
+						{
+							return InputPortConfig.GetPortGuid() == (*PortPtr)->GetPortGuid();
+						});
+
+					if (PortConfigPtr)
+					{
+						FDMXInputPortConfig PortConfig = *PortConfigPtr;
+						(*PortPtr)->UpdateFromConfig(PortConfig);
+					}
+				}
+			}
+		})
+);
+
+#undef DMX_OVERRIDE_INPUTPORT_VAR
+
+
 FDMXInputPortSharedRef FDMXInputPort::CreateFromConfig(FDMXInputPortConfig& InputPortConfig)
 {
@@ -52,37 +272,18 @@
 	UE_LOG(LogDMXProtocol, VeryVerbose, TEXT("Destroyed input port %s"), *PortName);
 }
 
-bool FDMXInputPort::CheckPriority(const int32 InPriority)
+FDMXInputPortConfig FDMXInputPort::MakeInputPortConfig() const
 {
-	if (InPriority > HighestReceivedPriority)
-	{
-		HighestReceivedPriority = InPriority;
-	}
-	
-	if (InPriority < LowestReceivedPriority)
-	{
-		LowestReceivedPriority = InPriority;
-	}
-	
-	switch (PriorityStrategy)
-	{
-	case(EDMXPortPriorityStrategy::None):
-		return true;
-	case(EDMXPortPriorityStrategy::HigherThan):
-		return InPriority > Priority;
-	case(EDMXPortPriorityStrategy::Equal):
-		return InPriority == Priority;
-	case(EDMXPortPriorityStrategy::LowerThan):
-		return InPriority < Priority;
-	case(EDMXPortPriorityStrategy::Highest):
-		return InPriority >= HighestReceivedPriority;
-	case(EDMXPortPriorityStrategy::Lowest):
-		return InPriority <= LowestReceivedPriority;
-	default:
-		break;
-	}
+	FDMXInputPortConfigParams Params;
+	Params.PortName = PortName;
+	Params.ProtocolName = Protocol.IsValid() ? Protocol->GetProtocolName() : NAME_None;
+	Params.CommunicationType = CommunicationType;
+	Params.DeviceAddress = DeviceAddress;
+	Params.LocalUniverseStart = LocalUniverseStart;
+	Params.NumUniverses = NumUniverses;
+	Params.ExternUniverseStart = ExternUniverseStart;
 
-	return false;
+	return FDMXInputPortConfig(PortGuid, Params);
 }
 
 void FDMXInputPort::UpdateFromConfig(FDMXInputPortConfig& InputPortConfig)
@@ -90,6 +291,13 @@
 	// Need a valid config for the port
 	InputPortConfig.MakeValid();
 
+	// Can only use configs that correspond to project settings
+	const UDMXProtocolSettings* ProtocolSettings = GetDefault<UDMXProtocolSettings>();
+	const bool bConfigIsInProjectSettings = ProtocolSettings->InputPortConfigs.ContainsByPredicate([&InputPortConfig](const FDMXInputPortConfig& Other) {
+		return InputPortConfig.GetPortGuid() == Other.GetPortGuid();
+		});
+	checkf(bConfigIsInProjectSettings, TEXT("Can only use configs with a guid that corresponds to a config in project settings"));
+
 	// Find if the port needs update its registration with the protocol
 	const bool bNeedsUpdateRegistration = [this, &InputPortConfig]()
 	{
@@ -143,15 +351,48 @@
 	OnPortUpdated.Broadcast();
 }
 
+bool FDMXInputPort::IsRegistered() const
+{
+	return bRegistered;
+}
+
 const FGuid& FDMXInputPort::GetPortGuid() const
 {
 	check(PortGuid.IsValid());
 	return PortGuid;
 }
 
-bool FDMXInputPort::IsRegistered() const
+bool FDMXInputPort::CheckPriority(const int32 InPriority)
 {
-	return bRegistered;
+	if (InPriority > HighestReceivedPriority)
+	{
+		HighestReceivedPriority = InPriority;
+	}
+
+	if (InPriority < LowestReceivedPriority)
+	{
+		LowestReceivedPriority = InPriority;
+	}
+
+	switch (PriorityStrategy)
+	{
+	case(EDMXPortPriorityStrategy::None):
+		return true;
+	case(EDMXPortPriorityStrategy::HigherThan):
+		return InPriority > Priority;
+	case(EDMXPortPriorityStrategy::Equal):
+		return InPriority == Priority;
+	case(EDMXPortPriorityStrategy::LowerThan):
+		return InPriority < Priority;
+	case(EDMXPortPriorityStrategy::Highest):
+		return InPriority >= HighestReceivedPriority;
+	case(EDMXPortPriorityStrategy::Lowest):
+		return InPriority <= LowestReceivedPriority;
+	default:
+		break;
+	}
+
+	return false;
 }
 
 void FDMXInputPort::AddRawListener(TSharedRef<FDMXRawListener> InRawListener)
@@ -321,25 +562,8 @@
 {
 	bReceiveDMXEnabled = bEnabled;
 
-	UpdateFromConfig(*FindInputPortConfigChecked());
+	FDMXInputPortConfig InputPortConfig = MakeInputPortConfig();
+	UpdateFromConfig(InputPortConfig);
 }
 
-FDMXInputPortConfig* FDMXInputPort::FindInputPortConfigChecked() const
-{
-	UDMXProtocolSettings* ProjectSettings = GetMutableDefault<UDMXProtocolSettings>();
-
-	for (FDMXInputPortConfig& InputPortConfig : ProjectSettings->InputPortConfigs)
-	{
-		if (InputPortConfig.GetPortGuid() == PortGuid)
-		{
-			return &InputPortConfig;
-		}
-	}
-
-	// Check failed, no config found
-	checkNoEntry();
-
-	return nullptr;
-}
-
 #undef LOCTEXT_NAMESPACE
diff --git a/Engine/Plugins/VirtualProduction/DMX/DMXProtocol/Source/DMXProtocol/Private/IO/DMXInputPortConfig.cpp b/Engine/Plugins/VirtualProduction/DMX/DMXProtocol/Source/DMXProtocol/Private/IO/DMXInputPortConfig.cpp
--- a/Engine/Plugins/VirtualProduction/DMX/DMXProtocol/Source/DMXProtocol/Private/IO/DMXInputPortConfig.cpp
+++ b/Engine/Plugins/VirtualProduction/DMX/DMXProtocol/Source/DMXProtocol/Private/IO/DMXInputPortConfig.cpp
@@ -18,6 +18,8 @@ FDMXInputPortConfigParams::FDMXInputPortConfigParams(const FDMXInputPortConfig&
 	, LocalUniverseStart(InputPortConfig.GetLocalUniverseStart())
 	, NumUniverses(InputPortConfig.GetNumUniverses())
 	, ExternUniverseStart(InputPortConfig.GetExternUniverseStart())
+	, PriorityStrategy(InputPortConfig.GetPortPriorityStrategy())
+	, Priority(InputPortConfig.GetPriority())
 {}
 
 
@@ -45,6 +47,8 @@ FDMXInputPortConfig::FDMXInputPortConfig(const FGuid& InPortGuid, const FDMXInpu
 	, LocalUniverseStart(InitializationData.LocalUniverseStart)
 	, NumUniverses(InitializationData.NumUniverses)
 	, ExternUniverseStart(InitializationData.ExternUniverseStart)
+	, PriorityStrategy(InitializationData.PriorityStrategy)
+	, Priority(InitializationData.Priority)
 	, PortGuid(InPortGuid)
 {
 	// Cannot create port configs before the protocol module is up (it is required to sanetize protocol names).
diff --git a/Engine/Plugins/VirtualProduction/DMX/DMXProtocol/Source/DMXProtocol/Private/IO/DMXOutputPort.cpp b/Engine/Plugins/VirtualProduction/DMX/DMXProtocol/Source/DMXProtocol/Private/IO/DMXOutputPort.cpp
--- a/Engine/Plugins/VirtualProduction/DMX/DMXProtocol/Source/DMXProtocol/Private/IO/DMXOutputPort.cpp
+++ b/Engine/Plugins/VirtualProduction/DMX/DMXProtocol/Source/DMXProtocol/Private/IO/DMXOutputPort.cpp
@@ -2,6 +2,7 @@
 
 #include "IO/DMXOutputPort.h"
 
+#include "DMXProtocolLog.h"
 #include "DMXProtocolSettings.h"
 #include "Interfaces/IDMXProtocol.h"
 #include "Interfaces/IDMXSender.h"
@@ -9,11 +10,358 @@
 #include "IO/DMXPortManager.h"
 #include "IO/DMXRawListener.h"
 
+#include "HAL/IConsoleManager.h"
 #include "HAL/RunnableThread.h"
 
 
 #define LOCTEXT_NAMESPACE "DMXOutputPort"
 
+/** Helper to override a member variable of an Output Port */
+#define DMX_OVERRIDE_OUTPUTPORT_VAR(MemberName, PortName, Value) \
+{ \
+	const FDMXOutputPortSharedRef* OutputPortPtr = FDMXPortManager::Get().GetOutputPorts().FindByPredicate([PortName](const FDMXOutputPortSharedPtr& OutputPort) \
+		{ \
+			return OutputPort->GetPortName() == PortName; \
+		}); \
+	if (OutputPortPtr) \
+	{ \
+		const FDMXOutputPortConfig OldOutputPortConfig = (*OutputPortPtr)->MakeOutputPortConfig(); \
+		FDMXOutputPortConfigParams OutputPortConfigParams = FDMXOutputPortConfigParams(OldOutputPortConfig); \
+		OutputPortConfigParams.MemberName = Value; \
+		FDMXOutputPortConfig NewOutputPortConfig((*OutputPortPtr)->GetPortGuid(), OutputPortConfigParams); \
+		(*OutputPortPtr)->UpdateFromConfig(NewOutputPortConfig); \
+	} \
+}
+
+static FAutoConsoleCommand GDMXSetOutputPortProtocolCommand(
+	TEXT("DMX.SetOutputPortProtocol"),
+	TEXT("DMX.SetOutputPortProtocol [PortName][ProtocolName]. Sets the protocol used by the output port. Example: DMX.SetOutputPortProtocol MyOutputPort sACN"),
+	FConsoleCommandWithArgsDelegate::CreateStatic(
+		[](const TArray<FString>& Args)
+		{
+			if (Args.Num() < 2)
+			{
+				return;
+			}
+
+			const FString& PortName = Args[0];
+			const FName ProtocolNameValue = FName(Args[1]);
+
+			if (IDMXProtocol::GetProtocolNames().Contains(ProtocolNameValue))
+			{
+				DMX_OVERRIDE_OUTPUTPORT_VAR(ProtocolName, PortName, ProtocolNameValue);
+			}
+		})
+);
+
+static FAutoConsoleCommand GDMXSetOutputPortCommunicationTypeCommand(
+	TEXT("DMX.SetOutputPortCommunicationType"),
+	TEXT("DMX.SetOutputPortCommunicationType [PortName][CommunicationType (1 = Broadcast, 2 = Multicast, 3 = Unicast)]. Sets the communication type of an output port. Example: DMX.SetOutputPortCommunicationType MyOutputPort 2"),
+	FConsoleCommandWithArgsDelegate::CreateStatic(
+		[](const TArray<FString>& Args)
+		{
+			if (Args.Num() < 2)
+			{
+				return;
+			}
+
+			const FString& PortName = Args[0];
+			const FString& CommunicationTypeValueString = Args[1];
+			uint8 CommunicationTypeValue;
+			if (LexTryParseString<uint8>(CommunicationTypeValue, *CommunicationTypeValueString))
+			{
+				const EDMXCommunicationType CommunicationTypeEnumValue = static_cast<EDMXCommunicationType>(CommunicationTypeValue);
+				if (CommunicationTypeEnumValue == EDMXCommunicationType::Broadcast ||
+					CommunicationTypeEnumValue == EDMXCommunicationType::Multicast ||
+					CommunicationTypeEnumValue == EDMXCommunicationType::Unicast)
+				{
+					const FDMXOutputPortSharedRef* PortPtr = FDMXPortManager::Get().GetOutputPorts().FindByPredicate([PortName](const FDMXOutputPortSharedPtr& OutputPort)
+						{
+							return OutputPort->GetPortName() == PortName;
+						});
+
+					if (PortPtr)
+					{
+						if (const IDMXProtocolPtr Protocol = (*PortPtr)->GetProtocol())
+						{
+							if (Protocol->GetOutputPortCommunicationTypes().Contains(CommunicationTypeEnumValue))
+							{
+								DMX_OVERRIDE_OUTPUTPORT_VAR(CommunicationType, PortName, CommunicationTypeEnumValue);
+							}
+						}
+					}
+				}
+			}
+		})
+);
+
+
+static FAutoConsoleCommand GDMXSetOutputPortOutputPortDeviceAddressCommand(
+	TEXT("DMX.SetOutputPortDeviceAddress"),
+	TEXT("DMX.SetOutputPortDeviceAddress [PortName][DeviceAddress]. Sets the Device Address of an output port, usually the network interface card IP address. Example: DMX.SetInputPortDeviceAddress MyOutputPort 123.45.67.89"),
+	FConsoleCommandWithArgsDelegate::CreateStatic(
+		[](const TArray<FString>& Args)
+		{
+			if (Args.Num() < 2)
+			{
+				return;
+			}
+
+			const FString& PortName = Args[0];
+			const FString& DeviceAddressValue = Args[1];
+
+			DMX_OVERRIDE_OUTPUTPORT_VAR(DeviceAddress, PortName, DeviceAddressValue);
+		})
+);
+
+static FAutoConsoleCommand GDMXSetOutputPortOutputPortDestinationAddressesCommand(
+	TEXT("DMX.SetOutputPortDestinationAddresses"),
+	TEXT("DMX.SetOutputPortDestinationAddresses [PortName][DestinationAddress1][DestinationAddress2][...][DestinationAddressN]. Sets the Destination Addresses of an output port. Example: DMX.SetInputPortDeviceAddress MyOutputPort 11.33.55.77 22.44.66.88"),
+	FConsoleCommandWithArgsDelegate::CreateStatic(
+		[](const TArray<FString>& Args)
+		{
+			if (Args.Num() < 2)
+			{
+				return;
+			}
+
+			const FString& PortName = Args[0];
+			TArray<FDMXOutputPortDestinationAddress>  DestinationAddressesValue;
+			for (int32 ArgIndex = 1; ArgIndex < Args.Num(); ArgIndex++)
+			{
+				const FDMXOutputPortDestinationAddress Address = FDMXOutputPortDestinationAddress(Args[ArgIndex]);
+				DestinationAddressesValue.Add(Address);
+			}
+
+			DMX_OVERRIDE_OUTPUTPORT_VAR(DestinationAddresses, PortName, DestinationAddressesValue);
+		})
+);
+
+static FAutoConsoleCommand GDMXSetOutputPortLoopbackToEngineCommand(
+	TEXT("DMX.SetOutputPortInputIntoEngine"),
+	TEXT("DMX.SetOutputSetInputIntoEngine [PortName][Flag]. Sets if the Output Port is input into the engine directly. Example: DMX.SetOutputPortInputIntoEngine MyOutputPort 1"),
+	FConsoleCommandWithArgsDelegate::CreateStatic(
+		[](const TArray<FString>& Args)
+		{
+			if (Args.Num() < 2)
+			{
+				return;
+			}
+
+			const FString& PortName = Args[0];
+			const bool bLoopbackToEngine = Args[1] == TEXT("1") || Args[1].Equals(TEXT("true"), ESearchCase::IgnoreCase);
+
+			DMX_OVERRIDE_OUTPUTPORT_VAR(bLoopbackToEngine, PortName, bLoopbackToEngine);
+		})
+);
+
+static FAutoConsoleCommand GDMXSetOutputPortLocalUniverseStartCommand(
+	TEXT("DMX.SetOutputPortLocalUniverseStart"),
+	TEXT("DMX.SetOutputPortLocalUniverseStart [PortName][Universe]. Sets the local universe start of the output port. Example: DMX.SetOutputPortLocalUniverseStart MyOutputPort 5"),
+	FConsoleCommandWithArgsDelegate::CreateStatic(
+		[](const TArray<FString>& Args)
+		{
+			if (Args.Num() < 2)
+			{
+				return;
+			}
+
+			const FString& PortName = Args[0];
+			const FString& LocalUniverseStartValueString = Args[1];
+			int32 LocalUniverseStartValue;
+			if (LexTryParseString<int32>(LocalUniverseStartValue, *LocalUniverseStartValueString))
+			{
+				if (LocalUniverseStartValue >= 0 && LocalUniverseStartValue <= DMX_MAX_UNIVERSE)
+				{
+					DMX_OVERRIDE_OUTPUTPORT_VAR(LocalUniverseStart, PortName, LocalUniverseStartValue);
+				}
+			}
+		})
+);
+
+static FAutoConsoleCommand GDMXSetOutputPortNumUniversesCommand(
+	TEXT("DMX.SetOutputPortNumUniverses"),
+	TEXT("DMX.SetOutputPortNumUniverses [PortName][Universe]. Sets the num universes of the output port. Example: DMX.SetOutputPortNumUniverses MyOutputPort 10"),
+	FConsoleCommandWithArgsDelegate::CreateStatic(
+		[](const TArray<FString>& Args)
+		{
+			if (Args.Num() < 2)
+			{
+				return;
+			}
+
+			const FString& PortName = Args[0];
+			const FString& NumUniversesValueString = Args[1];
+			int32 NumUniversesValue;
+			if (LexTryParseString<int32>(NumUniversesValue, *NumUniversesValueString))
+			{
+				if (NumUniversesValue >= 0 && NumUniversesValue <= DMX_MAX_UNIVERSE)
+				{
+					DMX_OVERRIDE_OUTPUTPORT_VAR(NumUniverses, PortName, NumUniversesValue);
+				}
+			}
+		})
+);
+
+static FAutoConsoleCommand GDMXSetOutputPortExternUniverseStartCommand(
+	TEXT("DMX.SetOutputPortExternUniverseStart"),
+	TEXT("DMX.SetOutputPortExternUniverseStart [PortName][Universe]. Sets the extern universe start of the output port. Example: DMX.SetOutputPortExternUniverseStart MyOutputPort 7"),
+	FConsoleCommandWithArgsDelegate::CreateStatic(
+		[](const TArray<FString>& Args)
+		{
+			if (Args.Num() < 2)
+			{
+				return;
+			}
+
+			const FString& PortName = Args[0];
+			const FString& ExternUniverseStartValueString = Args[1];
+			int32 ExternUniverseStartValue;
+			if (LexTryParseString<int32>(ExternUniverseStartValue, *ExternUniverseStartValueString))
+			{
+				if (ExternUniverseStartValue >= 0 && ExternUniverseStartValue <= DMX_MAX_UNIVERSE)
+				{
+					DMX_OVERRIDE_OUTPUTPORT_VAR(ExternUniverseStart, PortName, ExternUniverseStartValue);
+				}
+			}
+		})
+);
+
+static FAutoConsoleCommand GDMXSetOutputPortPriorityCommand(
+	TEXT("DMX.SetOutputPortPriority"),
+	TEXT("DMX.SetOutputPortPriority [PortName][Priority]. Sets the priority of the output port. Example: DMX.SetOutputPortPriority MyOutputPort 100"),
+	FConsoleCommandWithArgsDelegate::CreateStatic(
+		[](const TArray<FString>& Args)
+		{
+			if (Args.Num() < 2)
+			{
+				return;
+			}
+
+			const FString& PortName = Args[0];
+			const FString& PriorityValueString = Args[1];
+			int32 PriorityValue;
+			if (LexTryParseString<int32>(PriorityValue, *PriorityValueString))
+			{
+				DMX_OVERRIDE_OUTPUTPORT_VAR(Priority, PortName, PriorityValue);
+			}
+		})
+);
+
+static FAutoConsoleCommand GDMXSetOutputPortDelayCommand(
+	TEXT("DMX.SetOutputPortDelay"),
+	TEXT("DMX.SetOutputPortDelay [PortName][Delay]. Sets the delay of the output port (depending on the delay frame rate). Example: DMX.SetOutputPortDelay MyOutputPort 3.54"),
+	FConsoleCommandWithArgsDelegate::CreateStatic(
+		[](const TArray<FString>& Args)
+		{
+			if (Args.Num() < 2)
+			{
+				return;
+			}
+
+			const FString& PortName = Args[0];
+			const FString& DelayValueString = Args[1];
+			double DelayValue;
+			if (LexTryParseString<double>(DelayValue, *DelayValueString))
+			{
+				if (DelayValue >= 0.0)
+				{
+					const FDMXOutputPortSharedRef* PortPtr = FDMXPortManager::Get().GetOutputPorts().FindByPredicate([PortName](const FDMXOutputPortSharedPtr& OutputPort)
+						{
+							return OutputPort->GetPortName() == PortName;
+						});
+
+					if (PortPtr)
+					{
+						const FDMXOutputPortConfig PortConfig = (*PortPtr)->MakeOutputPortConfig();
+						const double DelaySeconds = FMath::Min(DelayValue * PortConfig.GetDelayFrameRate().AsDecimal(), 60.f);
+
+						DMX_OVERRIDE_OUTPUTPORT_VAR(DelaySeconds, PortName, DelayValue);
+					}
+				}
+			}
+		})
+);
+
+static FAutoConsoleCommand GDMXSetOutputPortDelayFrameRateCommand(
+	TEXT("DMX.SetOutputPortDelayFrameRate"),
+	TEXT("DMX.SetOutputPortDelayFrameRate [PortName][FrameRate]. Sets the frame rate of the delay of the output port (1.0 for seconds). Example: DMX.SetOutputPortDelayFrameRate MyOutputPort 33.3"),
+	FConsoleCommandWithArgsDelegate::CreateStatic(
+		[](const TArray<FString>& Args)
+		{
+			if (Args.Num() < 2)
+			{
+				return;
+			}
+
+			const FString& PortName = Args[0];
+			const FString& DelayFrameRateValueString = Args[1];
+			double DelayFrameRateValue;
+			if (LexTryParseString<double>(DelayFrameRateValue, *DelayFrameRateValueString))
+			{
+				if (DelayFrameRateValue > 0.0)
+				{
+					const FDMXOutputPortSharedRef* PortPtr = FDMXPortManager::Get().GetOutputPorts().FindByPredicate([PortName](const FDMXOutputPortSharedPtr& OutputPort)
+						{
+							return OutputPort->GetPortName() == PortName;
+						});
+
+					if (PortPtr)
+					{
+						const FDMXOutputPortConfig PortConfig = (*PortPtr)->MakeOutputPortConfig();
+						const double OldDelay = PortConfig.GetDelaySeconds() / PortConfig.GetDelayFrameRate().AsDecimal();
+						
+						const FFrameRate NewDelayFrameRate = FFrameRate(1, DelayFrameRateValue);
+						const double NewDelaySeconds = FMath::Min(OldDelay * NewDelayFrameRate.AsDecimal(), 60.f);
+
+						DMX_OVERRIDE_OUTPUTPORT_VAR(DelayFrameRate, PortName, NewDelayFrameRate);
+						DMX_OVERRIDE_OUTPUTPORT_VAR(DelaySeconds, PortName, NewDelaySeconds);
+					}
+				}
+			}
+		})
+);
+
+static FAutoConsoleCommand GDMXResetOutputPortToProjectSettings(
+	TEXT("DMX.ResetOutputPortToProjectSettings"),
+	TEXT("DMX.ResetOutputPortToProjectSettings [PortName]. Resets the output port to how it is defined in project settings. Example: DMX.ResetOutputPortToProjectSettings MyOutputPort"),
+	FConsoleCommandWithArgsDelegate::CreateStatic(
+		[](const TArray<FString>& Args)
+		{
+			if (Args.Num() < 1)
+			{
+				return;
+			}
+
+			const FString& PortName = Args[0];
+			const FDMXOutputPortSharedRef* PortPtr = FDMXPortManager::Get().GetOutputPorts().FindByPredicate([PortName](const FDMXOutputPortSharedPtr& OutputPort)
+				{
+					return OutputPort->GetPortName() == PortName;
+				});
+
+			if (PortPtr)
+			{
+				const UDMXProtocolSettings* ProtocolSettings = GetDefault<UDMXProtocolSettings>();
+				if (ProtocolSettings)
+				{
+					const FDMXOutputPortConfig* PortConfigPtr = ProtocolSettings->OutputPortConfigs.FindByPredicate([PortPtr](const FDMXOutputPortConfig& OutputPortConfig)
+						{
+							return OutputPortConfig.GetPortGuid() == (*PortPtr)->GetPortGuid();
+						});
+
+					if (PortConfigPtr)
+					{
+						FDMXOutputPortConfig PortConfig = *PortConfigPtr;
+						(*PortPtr)->UpdateFromConfig(PortConfig);
+					}
+				}
+			}
+		})
+);
+
+#undef DMX_OVERRIDE_OUTPUTPORT_VAR
+
+
 FDMXOutputPortSharedRef FDMXOutputPort::CreateFromConfig(FDMXOutputPortConfig& OutputPortConfig)
 {
 	// Port Configs are expected to have a valid guid always
@@ -59,17 +407,36 @@ FDMXOutputPort::~FDMXOutputPort()
 	UE_LOG(LogDMXProtocol, VeryVerbose, TEXT("Destroyed output port %s"), *PortName);
 }
 
+FDMXOutputPortConfig FDMXOutputPort::MakeOutputPortConfig() const
+{
+	FDMXOutputPortConfigParams Params;
+	Params.PortName = PortName;
+	Params.ProtocolName = Protocol.IsValid() ? Protocol->GetProtocolName() : NAME_None;
+	Params.CommunicationType = CommunicationType;
+	Params.DeviceAddress = DeviceAddress;
+	Params.DestinationAddresses = DestinationAddresses;
+	Params.bLoopbackToEngine = CommunicationDeterminator.IsLoopbackToEngineEnabled();
+	Params.LocalUniverseStart = LocalUniverseStart;
+	Params.NumUniverses = NumUniverses;
+	Params.ExternUniverseStart = ExternUniverseStart;
+	Params.Priority = Priority;
+	Params.DelaySeconds = DelaySeconds;
+	Params.DelayFrameRate = DelayFrameRate;
+
+	return FDMXOutputPortConfig(PortGuid, Params);
+}
+
 void FDMXOutputPort::UpdateFromConfig(FDMXOutputPortConfig& OutputPortConfig)
 {	
 	// Need a valid config for the port
 	OutputPortConfig.MakeValid();
 
-	// Can only use configs that are in project settings
+	// Can only use configs that correspond to project settings
 	const UDMXProtocolSettings* ProtocolSettings = GetDefault<UDMXProtocolSettings>();
 	const bool bConfigIsInProjectSettings = ProtocolSettings->OutputPortConfigs.ContainsByPredicate([&OutputPortConfig](const FDMXOutputPortConfig& Other) {
 		return OutputPortConfig.GetPortGuid() == Other.GetPortGuid();
 	});
-	checkf(bConfigIsInProjectSettings, TEXT("Can only use configs that are in project settings"));
+	checkf(bConfigIsInProjectSettings, TEXT("Can only use configs with a guid that corresponds to a config in project settings"));
 
 	// Find if the port needs update its registration with the protocol
 	const bool bNeedsUpdateRegistration = [this, &OutputPortConfig]()
@@ -119,6 +486,7 @@ void FDMXOutputPort::UpdateFromConfig(FDMXOutputPortConfig& OutputPortConfig)
 	PortName = OutputPortConfig.GetPortName();
 	Priority = OutputPortConfig.GetPriority();
 	DelaySeconds = OutputPortConfig.GetDelaySeconds();
+	DelayFrameRate = OutputPortConfig.GetDelayFrameRate();
 
 	CommunicationDeterminator.SetLoopbackToEngine(OutputPortConfig.NeedsLoopbackToEngine());
 
@@ -326,32 +694,16 @@ void FDMXOutputPort::OnSetSendDMXEnabled(bool bEnabled)
 {
 	CommunicationDeterminator.SetSendEnabled(bEnabled);
 
-	UpdateFromConfig(*FindOutputPortConfigChecked());
+	FDMXOutputPortConfig Config = MakeOutputPortConfig();
+	UpdateFromConfig(Config);
 }
 
 void FDMXOutputPort::OnSetReceiveDMXEnabled(bool bEnabled)
 {
 	CommunicationDeterminator.SetReceiveEnabled(bEnabled);
 
-	UpdateFromConfig(*FindOutputPortConfigChecked());
-}
-
-FDMXOutputPortConfig* FDMXOutputPort::FindOutputPortConfigChecked() const
-{
-	UDMXProtocolSettings* ProjectSettings = GetMutableDefault<UDMXProtocolSettings>();
-
-	for (FDMXOutputPortConfig& OutputPortConfig : ProjectSettings->OutputPortConfigs)
-	{
-		if (OutputPortConfig.GetPortGuid() == PortGuid)
-		{
-			return &OutputPortConfig;
-		}
-	}
-
-	// Check failed, no config found
-	checkNoEntry();
-
-	return nullptr;
+	FDMXOutputPortConfig Config = MakeOutputPortConfig();
+	UpdateFromConfig(Config);
 }
 
 bool FDMXOutputPort::Init()
diff --git a/Engine/Plugins/VirtualProduction/DMX/DMXProtocol/Source/DMXProtocol/Private/IO/DMXOutputPortConfig.cpp b/Engine/Plugins/VirtualProduction/DMX/DMXProtocol/Source/DMXProtocol/Private/IO/DMXOutputPortConfig.cpp
--- a/Engine/Plugins/VirtualProduction/DMX/DMXProtocol/Source/DMXProtocol/Private/IO/DMXOutputPortConfig.cpp
+++ b/Engine/Plugins/VirtualProduction/DMX/DMXProtocol/Source/DMXProtocol/Private/IO/DMXOutputPortConfig.cpp
@@ -22,6 +22,7 @@ FDMXOutputPortConfigParams::FDMXOutputPortConfigParams(const FDMXOutputPortConfi
 	, ExternUniverseStart(OutputPortConfig.GetExternUniverseStart())
 	, Priority(OutputPortConfig.GetPriority())
 	, DelaySeconds(OutputPortConfig.GetDelaySeconds())
+	, DelayFrameRate(OutputPortConfig.GetDelayFrameRate())
 {}
 
 
@@ -30,23 +31,24 @@ FDMXOutputPortConfig::FDMXOutputPortConfig()
 	, PortGuid(FGuid::NewGuid())
 {}
 
-FDMXOutputPortConfig::FDMXOutputPortConfig(const FGuid& InPortGuid)
-	: DelayFrameRate(FFrameRate(1.0, 1.0)) // Default delay frame rate to 1.0 (default to Seconds)
-	, PortGuid(InPortGuid)
-{
-	// Cannot create port configs before the protocol module is up (it is required to sanetize protocol names).
-	check(FModuleManager::Get().IsModuleLoaded("DMXProtocol"));
-	check(PortGuid.IsValid());
+	FDMXOutputPortConfig::FDMXOutputPortConfig(const FGuid& InPortGuid)
+		: DelayFrameRate(FFrameRate(1.0, 1.0)) // Default delay frame rate to 1.0 (default to Seconds)
+		, PortGuid(InPortGuid)
+	{
+		// Cannot create port configs before the protocol module is up (it is required to sanetize protocol names).
+		check(FModuleManager::Get().IsModuleLoaded("DMXProtocol"));
+		check(PortGuid.IsValid());
 
-	GenerateUniquePortName();
+		GenerateUniquePortName();
 
-	MakeValid();
-}
+		MakeValid();
+	}
 
 FDMXOutputPortConfig::FDMXOutputPortConfig(const FGuid& InPortGuid, const FDMXOutputPortConfigParams& InitializationData)
 	: PortName(InitializationData.PortName)
 	, ProtocolName(InitializationData.ProtocolName)
 	, CommunicationType(InitializationData.CommunicationType)
+	, DeviceAddress(InitializationData.DeviceAddress)
 	, DestinationAddresses(InitializationData.DestinationAddresses)
 	, bLoopbackToEngine(InitializationData.bLoopbackToEngine)
 	, LocalUniverseStart(InitializationData.LocalUniverseStart)
@@ -54,7 +56,7 @@ FDMXOutputPortConfig::FDMXOutputPortConfig(const FGuid& InPortGuid, const FDMXOu
 	, ExternUniverseStart(InitializationData.ExternUniverseStart)
 	, Priority(InitializationData.Priority)
 	, Delay(InitializationData.DelaySeconds)
-	, DelayFrameRate(FFrameRate(1.0, 1.0))
+	, DelayFrameRate(InitializationData.DelayFrameRate)
 	, PortGuid(InPortGuid)
 {
 	// Cannot create port configs before the protocol module is up (it is required to sanetize protocol names).
diff --git a/Engine/Plugins/VirtualProduction/DMX/DMXProtocol/Source/DMXProtocol/Public/IO/DMXInputPort.h b/Engine/Plugins/VirtualProduction/DMX/DMXProtocol/Source/DMXProtocol/Public/IO/DMXInputPort.h
--- a/Engine/Plugins/VirtualProduction/DMX/DMXProtocol/Source/DMXProtocol/Public/IO/DMXInputPort.h
+++ b/Engine/Plugins/VirtualProduction/DMX/DMXProtocol/Source/DMXProtocol/Public/IO/DMXInputPort.h
@@ -41,6 +41,9 @@ protected:
 public:
 	virtual ~FDMXInputPort();
 
+	/** Creates a dmx input port config that corresponds to the port */
+	FDMXInputPortConfig MakeInputPortConfig() const;
+
 	/** Updates the Port to use the config of the InputPortConfig. Makes the config valid if it's invalid. */
 	void UpdateFromConfig(FDMXInputPortConfig& InputPortConfig);
 
@@ -100,9 +103,6 @@ private:
 	/** Called to set if receive DMX should be enabled */
 	void OnSetReceiveDMXEnabled(bool bEnabled);
 
-	/** Returns the port config that corresponds to the guid of this port. */
-	FDMXInputPortConfig* FindInputPortConfigChecked() const;
-
 	/** The default buffer, which is being read on tick */
 	TQueue<FDMXSignalSharedPtr> DefaultInputQueue;
 
diff --git a/Engine/Plugins/VirtualProduction/DMX/DMXProtocol/Source/DMXProtocol/Public/IO/DMXInputPortConfig.h b/Engine/Plugins/VirtualProduction/DMX/DMXProtocol/Source/DMXProtocol/Public/IO/DMXInputPortConfig.h
--- a/Engine/Plugins/VirtualProduction/DMX/DMXProtocol/Source/DMXProtocol/Public/IO/DMXInputPortConfig.h
+++ b/Engine/Plugins/VirtualProduction/DMX/DMXProtocol/Source/DMXProtocol/Public/IO/DMXInputPortConfig.h
@@ -50,6 +50,8 @@ struct DMXPROTOCOL_API FDMXInputPortConfigParams
 	int32 LocalUniverseStart;
 	int32 NumUniverses;
 	int32 ExternUniverseStart;
+	EDMXPortPriorityStrategy PriorityStrategy;
+	int32 Priority;
 };
 
 /** 
diff --git a/Engine/Plugins/VirtualProduction/DMX/DMXProtocol/Source/DMXProtocol/Public/IO/DMXOutputPort.h b/Engine/Plugins/VirtualProduction/DMX/DMXProtocol/Source/DMXProtocol/Public/IO/DMXOutputPort.h
--- a/Engine/Plugins/VirtualProduction/DMX/DMXProtocol/Source/DMXProtocol/Public/IO/DMXOutputPort.h
+++ b/Engine/Plugins/VirtualProduction/DMX/DMXProtocol/Source/DMXProtocol/Public/IO/DMXOutputPort.h
@@ -43,15 +43,18 @@ struct FDMXOutputPortCommunicationDeterminator
 	/** Sets if there is a valid sender obj */
 	FORCEINLINE void SetHasValidSender(bool bInHasValidSender) { bHasValidSender = bInHasValidSender; }
 
-	/** Determinates if loopback to engine is needed. If true, loopback is needed */
+	/** Determinates if loopback to engine is needed (may be true even is loopback to engine is not enabled) */
 	FORCEINLINE bool NeedsLoopbackToEngine() const { return bLoopbackToEngine || !bReceiveEnabled || !bSendEnabled; }
 
-	/** Determinates if loopback to engine is needed. If true, loopback is needed */
+	/** Determinates if dmx needs to be sent (may be false even if send is enabled) */
 	FORCEINLINE bool NeedsSendDMX() const { return bSendEnabled && bHasValidSender; }
 
-	/** Determinates if loopback to engine is needed. If true, loopback is needed */
+	/** Determinates if send dmx is enabled */
 	FORCEINLINE bool IsSendDMXEnabled() const { return bSendEnabled; }
 
+	/** Returns true if loopback to engine is enabled (only true if loopback to engine is set enabled) */
+	FORCEINLINE bool IsLoopbackToEngineEnabled() const { return bLoopbackToEngine; }
+
 private:
 	bool bLoopbackToEngine;
 	bool bReceiveEnabled;
@@ -110,6 +113,9 @@ protected:
 public:
 	virtual ~FDMXOutputPort();
 
+	/** Creates a dmx output port config that corresponds to the port */
+	FDMXOutputPortConfig MakeOutputPortConfig() const;
+
 	/** Updates the Port to use the config of the OutputPortConfig */
 	void UpdateFromConfig(FDMXOutputPortConfig& OutputPortConfig);
 
@@ -165,9 +171,6 @@ protected:
 	/** Called to set if DMX should be enabled */
 	void OnSetReceiveDMXEnabled(bool bEnabled);
 
-	/** Returns the port config that corresponds to the guid of this port. */
-	FDMXOutputPortConfig* FindOutputPortConfigChecked() const;
-
 	//~ Begin FRunnable implementation
 	virtual bool Init() override;
 	virtual uint32 Run() override;
@@ -214,6 +217,9 @@ private:
 	/** Delay to apply on packets being sent */
 	double DelaySeconds = 0.0;
 
+	/** The frame rate of the delay */
+	FFrameRate DelayFrameRate;
+
 	/** Critical section required to be used when the SignalFragments are accessed across threads */
 	FCriticalSection AccessExternUniverseToLatestSignalMapCriticalSection;
 
diff --git a/Engine/Plugins/VirtualProduction/DMX/DMXProtocol/Source/DMXProtocol/Public/IO/DMXOutputPortConfig.h b/Engine/Plugins/VirtualProduction/DMX/DMXProtocol/Source/DMXProtocol/Public/IO/DMXOutputPortConfig.h
--- a/Engine/Plugins/VirtualProduction/DMX/DMXProtocol/Source/DMXProtocol/Public/IO/DMXOutputPortConfig.h
+++ b/Engine/Plugins/VirtualProduction/DMX/DMXProtocol/Source/DMXProtocol/Public/IO/DMXOutputPortConfig.h
@@ -95,6 +95,7 @@ public:
 	FORCEINLINE int32 GetPriority() const { return Priority; }
 	FORCEINLINE const FGuid& GetPortGuid() const { return PortGuid; }
 	int32 GetDelaySeconds() const;
+	FORCEINLINE const FFrameRate& GetDelayFrameRate() const { return DelayFrameRate; }
 
 	UE_DEPRECATED(5.0, "Output Ports now support many destination addresses. Use FDMXOutputPortConfig::GetDestinationAddresses instead. ")
 	FORCEINLINE const FString& GetDestinationAddress() const { return DestinationAddress; }
