DMX - Improve various UI elements and code formating

diff --git a/Engine/Plugins/VirtualProduction/DMX/DMXEngine/Source/DMXEditor/Private/Style/DMXEditorStyle.cpp b/Engine/Plugins/VirtualProduction/DMX/DMXEngine/Source/DMXEditor/Private/Style/DMXEditorStyle.cpp
--- a/Engine/Plugins/VirtualProduction/DMX/DMXEngine/Source/DMXEditor/Private/Style/DMXEditorStyle.cpp
+++ b/Engine/Plugins/VirtualProduction/DMX/DMXEngine/Source/DMXEditor/Private/Style/DMXEditorStyle.cpp
@@ -104,11 +104,6 @@ TSharedRef<FSlateStyleSet> FDMXEditorStyle::Create()
 		Style->Set("DMXEditor.PixelMapping.DistributionGrid.3.3", new IMAGE_BRUSH("Icons/DistributionGrid/icon_PixelDirection_3_3_161x160", Icon34x29));
 	}
 	
-	// DMX Editor icons
-	{
-		Style->Set("DMXEditor.FixtureTypesEditor.DeleteCellAttribute", new IMAGE_BRUSH("Icons/icon_Delete_CellAttribute_40x40", Icon40x40));
-	}
-
 	// Output Console
 	{
 		Style->Set("DMXEditor.OutputConsole.Fader", new FSlateColorBrush(FLinearColor(.2f, .2f, .2f, .2f)));
diff --git a/Engine/Plugins/VirtualProduction/DMX/DMXEngine/Source/DMXEditor/Private/Widgets/FixtureType/SDMXFixtureTypeMatrixFunctionsEditor.cpp b/Engine/Plugins/VirtualProduction/DMX/DMXEngine/Source/DMXEditor/Private/Widgets/FixtureType/SDMXFixtureTypeMatrixFunctionsEditor.cpp
--- a/Engine/Plugins/VirtualProduction/DMX/DMXEngine/Source/DMXEditor/Private/Widgets/FixtureType/SDMXFixtureTypeMatrixFunctionsEditor.cpp
+++ b/Engine/Plugins/VirtualProduction/DMX/DMXEngine/Source/DMXEditor/Private/Widgets/FixtureType/SDMXFixtureTypeMatrixFunctionsEditor.cpp
@@ -62,7 +62,8 @@ void SDMXFixtureTypeMatrixFunctionsEditor::Construct(const FArguments& InArgs, c
 
 TSharedRef<SHeaderRow> SDMXFixtureTypeMatrixFunctionsEditor::GenerateHeaderRow()
 {
-	const float StatusColumnWidth = FMath::Max(FEditorStyle::GetBrush("Icons.Warning")->GetImageSize().X + 6.f, FEditorStyle::GetBrush("Icons.Error")->GetImageSize().X + 6.f);
+	const float StatusColumnWidth = FMath::Max(FEditorStyle::GetBrush("Icons.Warning")->GetImageSize().X + 8.f, FEditorStyle::GetBrush("Icons.Error")->GetImageSize().X + 8.f);
+	const float DeleteAttributeColumnWidth = FEditorStyle::GetBrush("Icons.Delete")->GetImageSize().X + 8.f;
 
 	TAttribute<EVisibility> HeaderRowVisiblityAttribute = TAttribute<EVisibility>::Create(TAttribute<EVisibility>::FGetter::CreateLambda([this]()
 		{
@@ -97,7 +98,7 @@ TSharedRef<SHeaderRow> SDMXFixtureTypeMatrixFunctionsEditor::GenerateHeaderRow()
 		SHeaderRow::FColumn::FArguments()
 		.ColumnId(FDMXFixtureTypeMatrixFunctionsEditorCollumnIDs::DeleteAttribute)
 		.DefaultLabel(LOCTEXT("DeleteAttributeColumnLabel", ""))
-		.FixedWidth(42.f)
+		.FixedWidth(DeleteAttributeColumnWidth)
 	);
 
 	return HeaderRow;
diff --git a/Engine/Plugins/VirtualProduction/DMX/DMXEngine/Source/DMXEditor/Private/Widgets/FixtureType/SDMXFixtureTypeMatrixFunctionsEditorMatrixRow.cpp b/Engine/Plugins/VirtualProduction/DMX/DMXEngine/Source/DMXEditor/Private/Widgets/FixtureType/SDMXFixtureTypeMatrixFunctionsEditorMatrixRow.cpp
--- a/Engine/Plugins/VirtualProduction/DMX/DMXEngine/Source/DMXEditor/Private/Widgets/FixtureType/SDMXFixtureTypeMatrixFunctionsEditorMatrixRow.cpp
+++ b/Engine/Plugins/VirtualProduction/DMX/DMXEngine/Source/DMXEditor/Private/Widgets/FixtureType/SDMXFixtureTypeMatrixFunctionsEditorMatrixRow.cpp
@@ -92,22 +92,26 @@ TSharedRef<SWidget> SDMXFixtureTypeMatrixFunctionsEditorMatrixRow::GenerateWidge
 	else if (ColumnName == FDMXFixtureTypeMatrixFunctionsEditorCollumnIDs::DeleteAttribute)
 	{
 		return
-			SNew(SBorder)
+			SNew(SButton)
+			.ContentPadding(2.0f)
+			.ButtonStyle(FEditorStyle::Get(), "HoverHintOnly")
+			.OnClicked(this, &SDMXFixtureTypeMatrixFunctionsEditorMatrixRow::OnDeleteCellAttributeClicked)
+			.ToolTipText(LOCTEXT("RemoveCellAttributeTooltip", "Removes the Cell Attribute"))
+			.ForegroundColor(FSlateColor::UseForeground())
 			.HAlign(HAlign_Center)
 			.VAlign(VAlign_Center)
-			.BorderImage(FEditorStyle::GetBrush("NoBorder"))
+			.Content()
 			[
 				SNew(SImage)
-				.Image(FDMXEditorStyle::Get().GetBrush("DMXEditor.FixtureTypesEditor.DeleteCellAttribute"))
-				.DesiredSizeOverride(FVector2D(20.f, 20.f))
-				.OnMouseButtonDown(this, &SDMXFixtureTypeMatrixFunctionsEditorMatrixRow::OnDeleteCellAttributeClicked)
+				.Image(FEditorStyle::GetBrush("Icons.Delete"))
+				.ColorAndOpacity(FLinearColor(0.6f, 0.6f, 0.6f))
 			];
 	}
 
 	return SNullWidget::NullWidget;
 }
 
-FReply SDMXFixtureTypeMatrixFunctionsEditorMatrixRow::OnDeleteCellAttributeClicked(const FGeometry& MyGeometry, const FPointerEvent& MouseEvent)
+FReply SDMXFixtureTypeMatrixFunctionsEditorMatrixRow::OnDeleteCellAttributeClicked()
 {
 	OnRequestDelete.ExecuteIfBound();
 
diff --git a/Engine/Plugins/VirtualProduction/DMX/DMXEngine/Source/DMXEditor/Private/Widgets/FixtureType/SDMXFixtureTypeMatrixFunctionsEditorMatrixRow.h b/Engine/Plugins/VirtualProduction/DMX/DMXEngine/Source/DMXEditor/Private/Widgets/FixtureType/SDMXFixtureTypeMatrixFunctionsEditorMatrixRow.h
--- a/Engine/Plugins/VirtualProduction/DMX/DMXEngine/Source/DMXEditor/Private/Widgets/FixtureType/SDMXFixtureTypeMatrixFunctionsEditorMatrixRow.h
+++ b/Engine/Plugins/VirtualProduction/DMX/DMXEngine/Source/DMXEditor/Private/Widgets/FixtureType/SDMXFixtureTypeMatrixFunctionsEditorMatrixRow.h
@@ -36,7 +36,7 @@ protected:
 
 private:
 	/** Called when the delete Cell Attribute Button was clicked */
-	FReply OnDeleteCellAttributeClicked(const FGeometry& MyGeometry, const FPointerEvent& MouseEvent);
+	FReply OnDeleteCellAttributeClicked();
 
 	/** Returns true if the Attribute Name exists and is not None */
 	bool HasInvalidCellAttributeName() const;
diff --git a/Engine/Plugins/VirtualProduction/DMX/DMXProtocol/Source/DMXProtocol/Private/IO/DMXInputPort.cpp b/Engine/Plugins/VirtualProduction/DMX/DMXProtocol/Source/DMXProtocol/Private/IO/DMXInputPort.cpp
--- a/Engine/Plugins/VirtualProduction/DMX/DMXProtocol/Source/DMXProtocol/Private/IO/DMXInputPort.cpp
+++ b/Engine/Plugins/VirtualProduction/DMX/DMXProtocol/Source/DMXProtocol/Private/IO/DMXInputPort.cpp
@@ -177,7 +177,7 @@ static FAutoConsoleCommand GDMXSetInputPortPriorityStrategyCommand(
 
 static FAutoConsoleCommand GDMXSetInputPortPriorityCommand(
 	TEXT("DMX.SetInputPortPriority"),
-	TEXT("DMX.SetInputPortPriority [PortName][Universe]. Sets the priority of the input port. Example: DMX.SetInputPortPriority MyInputPort 100"),
+	TEXT("DMX.SetInputPortPriority [PortName][Priority]. Sets the priority of the input port. Example: DMX.SetInputPortPriority MyInputPort 100"),
 	FConsoleCommandWithArgsDelegate::CreateStatic(
 		[](const TArray<FString>& Args)
 		{
diff --git a/Engine/Plugins/VirtualProduction/DMX/DMXProtocol/Source/DMXProtocol/Private/IO/DMXOutputPort.cpp b/Engine/Plugins/VirtualProduction/DMX/DMXProtocol/Source/DMXProtocol/Private/IO/DMXOutputPort.cpp
--- a/Engine/Plugins/VirtualProduction/DMX/DMXProtocol/Source/DMXProtocol/Private/IO/DMXOutputPort.cpp
+++ b/Engine/Plugins/VirtualProduction/DMX/DMXProtocol/Source/DMXProtocol/Private/IO/DMXOutputPort.cpp
@@ -35,7 +35,7 @@
 
 static FAutoConsoleCommand GDMXSetOutputPortProtocolCommand(
 	TEXT("DMX.SetOutputPortProtocol"),
-	TEXT("DMX.SetOutputPortProtocol [PortName][ProtocolName]. Sets the protocol used by the output port. Example: DMX.SetOutputPortProtocol MyOutputPort sACN"),
+	TEXT("DMX.SetOutputPortProtocol [PortName][ProtocolName]. Sets the protocol used by the output port. Example: DMX.SetOutputPortProtocol MyOutputPort Art-Net"),
 	FConsoleCommandWithArgsDelegate::CreateStatic(
 		[](const TArray<FString>& Args)
 		{
diff --git a/Engine/Plugins/VirtualProduction/DMX/DMXProtocol/Source/DMXProtocolEditor/Private/DetailsCustomizations/DMXOutputPortDestinationAddressCustomization.cpp b/Engine/Plugins/VirtualProduction/DMX/DMXProtocol/Source/DMXProtocolEditor/Private/DetailsCustomizations/DMXOutputPortDestinationAddressCustomization.cpp
--- a/Engine/Plugins/VirtualProduction/DMX/DMXProtocol/Source/DMXProtocolEditor/Private/DetailsCustomizations/DMXOutputPortDestinationAddressCustomization.cpp
+++ b/Engine/Plugins/VirtualProduction/DMX/DMXProtocol/Source/DMXProtocolEditor/Private/DetailsCustomizations/DMXOutputPortDestinationAddressCustomization.cpp
@@ -8,6 +8,7 @@
 #include "IDetailChildrenBuilder.h"
 #include "IDetailPropertyRow.h" 
 #include "DetailWidgetRow.h"
+#include "Widgets/Images/SImage.h"
 
 
 #define LOCTEXT_NAMESPACE "DMXOutputPortDestinationAddressCustomization"
@@ -26,13 +27,69 @@ void FDMXOutputPortDestinationAddressCustomization::CustomizeChildren(TSharedRef
 {
 	DestinationAddressStringHandle = StructPropertyHandle->GetChildHandle(GET_MEMBER_NAME_CHECKED(FDMXOutputPortDestinationAddress, DestinationAddressString));
 
+	// Create the content depending on the property being an array member 
+	TSharedRef<SHorizontalBox> ContentBox = SNew(SHorizontalBox);
+
+	ContentBox->AddSlot()
+		.HAlign(HAlign_Left)
+		.VAlign(VAlign_Center)
+		.AutoWidth()
+		[
+			SNew(SBox)
+			.MinDesiredWidth(124.f)
+			[
+				SAssignNew(IPAddressEditWidget, SDMXIPAddressEditWidget)
+				.InitialValue(GetIPAddress())
+				.bShowLocalNICComboBox(false)
+				.OnIPAddressSelected(this, &FDMXOutputPortDestinationAddressCustomization::OnDestinationAddressSelected)
+			]
+		];
+
+	const bool bIsArrayMember = [this]()
+	{
+		if (const TSharedPtr<IPropertyHandle> DestinationAddressStructHandle = DestinationAddressStringHandle->GetParentHandle())
+		{
+			if (TSharedPtr<IPropertyHandle> DestinationAddressesHandle = DestinationAddressStructHandle->GetParentHandle())
+			{
+				if (TSharedPtr<IPropertyHandleArray> DestinationAddressesHandleArray = DestinationAddressesHandle->AsArray())
+				{
+					return true;
+				}
+			}
+		}
+
+		return false;
+	}();
+
+	if (bIsArrayMember)
+	{
+		ContentBox->AddSlot()
+			.HAlign(HAlign_Left)
+			.VAlign(VAlign_Center)
+			.AutoWidth()
+			.Padding(FMargin(4.f, 0.f, 0.f, 0.f))
+			[
+				SNew(SButton)
+				.ContentPadding(2.0f)
+				.ButtonStyle(FEditorStyle::Get(), "HoverHintOnly")
+				.OnClicked(this, &FDMXOutputPortDestinationAddressCustomization::OnDeleteDestinationAddressClicked)
+				.ToolTipText(LOCTEXT("RemoveCellAttributeTooltip", "Removes the Destination Address"))
+				.ForegroundColor(FSlateColor::UseForeground())
+				.HAlign(HAlign_Center)
+				.VAlign(VAlign_Center)
+				.Content()
+				[
+					SNew(SImage)
+					.Image(FEditorStyle::GetBrush("Icons.Delete"))
+					.ColorAndOpacity(FLinearColor(0.6f, 0.6f, 0.6f))
+				]
+			];
+	}
+
 	ChildBuilder.AddCustomRow(LOCTEXT("InvalidPortConfigSearchString", "Invalid"))
 		.ValueContent()
 		[
-			SAssignNew(IPAddressEditWidget, SDMXIPAddressEditWidget)
-			.InitialValue(GetIPAddress())
-			.bShowLocalNICComboBox(false)
-			.OnIPAddressSelected(this, &FDMXOutputPortDestinationAddressCustomization::OnIPAddressSelected)
+			ContentBox
 		];
 }
 
@@ -54,12 +111,12 @@ void FDMXOutputPortDestinationAddressCustomization::SetIPAddress(const FString&
 	DestinationAddressStringHandle->NotifyPostChange(EPropertyChangeType::ValueSet);
 }
 
-void FDMXOutputPortDestinationAddressCustomization::OnIPAddressSelected()
+void FDMXOutputPortDestinationAddressCustomization::OnDestinationAddressSelected()
 {
 	TArray<UObject*> OuterObjects;
 	DestinationAddressStringHandle->GetOuterObjects(OuterObjects);
 	
-	// Forward the change to the outer object so it gets the property change. Required due to get the change from the nested property all the way to UDMXProtocolSettings.
+	// Forward the change to the outer object so it gets the property change. Required to get the change from the nested property all the way to the property owner (e.g. UDMXProtocolSettings)
 	for (UObject* Object : OuterObjects)
 	{
 		Object->PreEditChange(FDMXOutputPortDestinationAddress::StaticStruct()->FindPropertyByName(GET_MEMBER_NAME_CHECKED(FDMXOutputPortDestinationAddress, DestinationAddressString)));
@@ -74,4 +131,53 @@ void FDMXOutputPortDestinationAddressCustomization::OnIPAddressSelected()
 	}
 }
 
+FReply FDMXOutputPortDestinationAddressCustomization::OnDeleteDestinationAddressClicked()
+{
+	if (const TSharedPtr<IPropertyHandle> DestinationAddressStructHandle = DestinationAddressStringHandle->GetParentHandle())
+	{
+		if (TSharedPtr<IPropertyHandle> DestinationAddressesHandle = DestinationAddressStructHandle->GetParentHandle())
+		{
+			if (TSharedPtr<IPropertyHandleArray> DestinationAddressesHandleArray = DestinationAddressesHandle->AsArray())
+			{
+				// Find the index of this property
+				const uint32 MyIndexInArray = [&DestinationAddressStructHandle, &DestinationAddressesHandleArray, this]() -> uint32
+				{
+					uint32 NumElements;
+					DestinationAddressesHandleArray->GetNumElements(NumElements);
+					for (uint32 Index = 0; Index < NumElements; Index++)
+					{
+						if (DestinationAddressStructHandle->GetIndexInArray() == Index)
+						{
+							return Index;
+						}
+					}
+
+					return INDEX_NONE;
+				}();
+
+				if (MyIndexInArray != INDEX_NONE)
+				{
+					TArray<UObject*> OuterObjects;
+					DestinationAddressStringHandle->GetOuterObjects(OuterObjects);
+
+					// Forward the change to the outer object so it gets the property change. Required to get the change from the nested property all the way to the property owner (e.g. UDMXProtocolSettings)
+					for (UObject* Object : OuterObjects)
+					{
+						Object->PreEditChange(FDMXOutputPortDestinationAddress::StaticStruct()->FindPropertyByName(GET_MEMBER_NAME_CHECKED(FDMXOutputPortDestinationAddress, DestinationAddressString)));
+					}
+
+					DestinationAddressesHandleArray->DeleteItem(MyIndexInArray);
+
+					for (UObject* Object : OuterObjects)
+					{
+						Object->PostEditChange();
+					}
+				}
+			}
+		}
+	}
+
+	return FReply::Handled();
+}
+
 #undef LOCTEXT_NAMESPACE
diff --git a/Engine/Plugins/VirtualProduction/DMX/DMXProtocol/Source/DMXProtocolEditor/Private/DetailsCustomizations/DMXOutputPortDestinationAddressCustomization.h b/Engine/Plugins/VirtualProduction/DMX/DMXProtocol/Source/DMXProtocolEditor/Private/DetailsCustomizations/DMXOutputPortDestinationAddressCustomization.h
--- a/Engine/Plugins/VirtualProduction/DMX/DMXProtocol/Source/DMXProtocolEditor/Private/DetailsCustomizations/DMXOutputPortDestinationAddressCustomization.h
+++ b/Engine/Plugins/VirtualProduction/DMX/DMXProtocol/Source/DMXProtocolEditor/Private/DetailsCustomizations/DMXOutputPortDestinationAddressCustomization.h
@@ -30,8 +30,11 @@ private:
 	/** Sets the IP Adderss */
 	void SetIPAddress(const FString& NewIPAddress);
 
-	/** Called when an IP Address was selcted */
-	void OnIPAddressSelected();
+	/** Called when a destination Address was selected */
+	void OnDestinationAddressSelected();
+
+	/** Called when a destination address wants to be deleted */
+	FReply OnDeleteDestinationAddressClicked();
 
 	/** Widget to enter an IP address */
 	TSharedPtr<SDMXIPAddressEditWidget> IPAddressEditWidget;
