DMX - Fix thread-safety in Art-Net and sACN sender, improve internal loopback to not read paritally written universes

diff --git a/Engine/Plugins/VirtualProduction/DMX/DMXProtocol/Source/DMXProtocol/Private/IO/DMXOutputPort.cpp b/Engine/Plugins/VirtualProduction/DMX/DMXProtocol/Source/DMXProtocol/Private/IO/DMXOutputPort.cpp
--- a/Engine/Plugins/VirtualProduction/DMX/DMXProtocol/Source/DMXProtocol/Private/IO/DMXOutputPort.cpp
+++ b/Engine/Plugins/VirtualProduction/DMX/DMXProtocol/Source/DMXProtocol/Private/IO/DMXOutputPort.cpp
@@ -507,15 +507,20 @@ const FGuid& FDMXOutputPort::GetPortGuid() const
 
 void FDMXOutputPort::ClearBuffers()
 {
+	check(IsInGameThread());
+
 	for (const TSharedRef<FDMXRawListener>& RawListener : RawListeners)
 	{
 		RawListener->ClearBuffer();
 	}
 
-	FScopeLock LockClearBuffers(&ClearBuffersCriticalSection);
+	// Clear gamethread buffer
+	ExternUniverseToLatestSignalMap_GameThread.Reset();
 
+	// Clear port thread buffers
+	FScopeLock LockClearBuffers(&ClearBuffersCriticalSection);
 	SignalFragments.Empty();
-	ExternUniverseToLatestSignalMap.Reset();
+	ExternUniverseToLatestSignalMap_PortThread.Reset();
 }
 
 bool FDMXOutputPort::IsLoopbackToEngine() const
@@ -537,17 +542,14 @@ TArray<FString> FDMXOutputPort::GetDestinationAddresses() const
 
 bool FDMXOutputPort::GameThreadGetDMXSignal(int32 LocalUniverseID, FDMXSignalSharedPtr& OutDMXSignal, bool bEvenIfNotLoopbackToEngine)
 {
-#if UE_BUILD_DEBUG
 	check(IsInGameThread());
-#endif // UE_BUILD_DEBUG
 
 	const bool bNeedsLoopbackToEngine = CommunicationDeterminator.NeedsLoopbackToEngine();
 	if (bNeedsLoopbackToEngine || bEvenIfNotLoopbackToEngine)
 	{
 		int32 ExternUniverseID = ConvertLocalToExternUniverseID(LocalUniverseID);
 
-		FScopeLock LockExternUniverseToLatestSignalMap(&AccessExternUniverseToLatestSignalMapCriticalSection);
-		const FDMXSignalSharedPtr* SignalPtr = ExternUniverseToLatestSignalMap.Find(ExternUniverseID);
+		const FDMXSignalSharedPtr* SignalPtr = ExternUniverseToLatestSignalMap_GameThread.Find(ExternUniverseID);
 		if (SignalPtr)
 		{
 			OutDMXSignal = *SignalPtr;
@@ -561,15 +563,12 @@ bool FDMXOutputPort::GameThreadGetDMXSignal(int32 LocalUniverseID, FDMXSignalSha
 bool FDMXOutputPort::GameThreadGetDMXSignalFromRemoteUniverse(FDMXSignalSharedPtr& OutDMXSignal, int32 RemoteUniverseID, bool bEvenIfNotLoopbackToEngine)
 {
 	// DEPRECATED 4.27
-#if UE_BUILD_DEBUG
 	check(IsInGameThread());
-#endif // UE_BUILD_DEBUG
 
 	const bool bNeedsLoopbackToEngine = CommunicationDeterminator.NeedsLoopbackToEngine();
 	if (bNeedsLoopbackToEngine || bEvenIfNotLoopbackToEngine)
 	{
-		FScopeLock LockExternUniverseToLatestSignalMap(&AccessExternUniverseToLatestSignalMapCriticalSection);
-		const FDMXSignalSharedPtr* SignalPtr = ExternUniverseToLatestSignalMap.Find(RemoteUniverseID);
+		const FDMXSignalSharedPtr* SignalPtr = ExternUniverseToLatestSignalMap_GameThread.Find(RemoteUniverseID);
 		if (SignalPtr)
 		{
 			OutDMXSignal = *SignalPtr;
@@ -613,9 +612,7 @@ void FDMXOutputPort::RemoveRawListener(TSharedRef<FDMXRawListener> InRawListener
 
 void FDMXOutputPort::SendDMX(int32 LocalUniverseID, const TMap<int32, uint8>& ChannelToValueMap)
 {
-#if UE_BUILD_DEBUG
 	check(IsInGameThread());
-#endif // UE_BUILD_DEBUG
 
 	if (IsLocalUniverseInPortRange(LocalUniverseID))
 	{
@@ -628,9 +625,28 @@ void FDMXOutputPort::SendDMX(int32 LocalUniverseID, const TMap<int32, uint8>& Ch
 			const int32 ExternUniverseID = ConvertLocalToExternUniverseID(LocalUniverseID);
 
 			const double SendTime = FPlatformTime::Seconds() + DelaySeconds;
-			const TSharedPtr<FDMXSignalFragment> Fragment = MakeShared<FDMXSignalFragment>(ExternUniverseID, ChannelToValueMap, SendTime);
 
+			// Enqueue for this port's thread
+			const TSharedPtr<FDMXSignalFragment> Fragment = MakeShared<FDMXSignalFragment>(ExternUniverseID, ChannelToValueMap, SendTime);
 			SignalFragments.Enqueue(Fragment);
+
+			if (bNeedsLoopbackToEngine)
+			{
+				// Write the fragment to the game thread's buffer
+				const FDMXSignalSharedPtr& Signal = ExternUniverseToLatestSignalMap_GameThread.FindOrAdd(ExternUniverseID, MakeShared<FDMXSignal, ESPMode::ThreadSafe>());
+
+				for (const TTuple<int32, uint8>& ChannelValueKvp : ChannelToValueMap)
+				{
+					int32 ChannelIndex = ChannelValueKvp.Key - 1;
+
+					// Filter invalid indicies so we can send bp calls here without testing them first.
+					if (Signal->ChannelData.IsValidIndex(ChannelIndex))
+					{
+						Signal->Timestamp = SendTime;
+						Signal->ChannelData[ChannelIndex] = ChannelValueKvp.Value;
+					}
+				}
+			}
 		}
 	}
 }
@@ -638,9 +654,7 @@ void FDMXOutputPort::SendDMX(int32 LocalUniverseID, const TMap<int32, uint8>& Ch
 void FDMXOutputPort::SendDMXToRemoteUniverse(const TMap<int32, uint8>& ChannelToValueMap, int32 RemoteUniverse)
 {
 	// DEPRECATED 4.27
-#if UE_BUILD_DEBUG
 	check(IsInGameThread());
-#endif // UE_BUILD_DEBUG
 
 	if (IsExternUniverseInPortRange(RemoteUniverse))
 	{
@@ -651,9 +665,28 @@ void FDMXOutputPort::SendDMXToRemoteUniverse(const TMap<int32, uint8>& ChannelTo
 		if (bNeedsSendDMX || bNeedsLoopbackToEngine)
 		{
 			const double SendTime = FPlatformTime::Seconds() + DelaySeconds;
-			const TSharedPtr<FDMXSignalFragment> Fragment = MakeShared<FDMXSignalFragment>(RemoteUniverse, ChannelToValueMap, SendTime);
 
+			// Enqueue for this port's thread
+			const TSharedPtr<FDMXSignalFragment> Fragment = MakeShared<FDMXSignalFragment>(RemoteUniverse, ChannelToValueMap, SendTime);
 			SignalFragments.Enqueue(Fragment);
+
+			if (bNeedsLoopbackToEngine)
+			{
+				// Write the fragment to the game thread's buffer
+				const FDMXSignalSharedPtr& Signal = ExternUniverseToLatestSignalMap_GameThread.FindOrAdd(RemoteUniverse, MakeShared<FDMXSignal, ESPMode::ThreadSafe>());
+
+				for (const TTuple<int32, uint8>& ChannelValueKvp : ChannelToValueMap)
+				{
+					int32 ChannelIndex = ChannelValueKvp.Key - 1;
+
+					// Filter invalid indicies so we can send bp calls here without testing them first.
+					if (Signal->ChannelData.IsValidIndex(ChannelIndex))
+					{
+						Signal->Timestamp = SendTime;
+						Signal->ChannelData[ChannelIndex] = ChannelValueKvp.Value;
+					}
+				}
+			}
 		}
 	}
 }
@@ -776,16 +809,9 @@ void FDMXOutputPort::ProcessSendDMX()
 			{
 				if (OldestFragment->SendTime <= Now)
 				{
-					FDMXSignalSharedPtr* SignalPtr = ExternUniverseToLatestSignalMap.Find(OldestFragment->ExternUniverseID);
-					if (!SignalPtr)
-					{
-						FScopeLock LockExternUniverseToLatestSignalMap(&AccessExternUniverseToLatestSignalMapCriticalSection);
-						SignalPtr = &ExternUniverseToLatestSignalMap.Add(OldestFragment->ExternUniverseID, MakeShared<FDMXSignal, ESPMode::ThreadSafe>());
-					}
-					checkf(SignalPtr, TEXT("Cannot find or add a DMX signal to the DMX Output Port buffer"));
+					const FDMXSignalSharedPtr& Signal = ExternUniverseToLatestSignalMap_PortThread.FindOrAdd(OldestFragment->ExternUniverseID, MakeShared<FDMXSignal, ESPMode::ThreadSafe>());
 
 					// Write the fragment & meta data 
-					FDMXSignalSharedPtr& Signal = *SignalPtr;
 					for (const TTuple<int32, uint8>& ChannelValueKvp : OldestFragment->ChannelToValueMap)
 					{
 						int32 ChannelIndex = ChannelValueKvp.Key - 1;
@@ -815,7 +841,7 @@ void FDMXOutputPort::ProcessSendDMX()
 
 	// Send new and alive DMX Signals
 	const bool bNeedsSendDMX = CommunicationDeterminator.NeedsSendDMX();
-	for (const TTuple<int32, FDMXSignalSharedPtr>& UniverseToSignalPair : ExternUniverseToLatestSignalMap)
+	for (const TTuple<int32, FDMXSignalSharedPtr>& UniverseToSignalPair : ExternUniverseToLatestSignalMap_PortThread)
 	{
 		if (UniverseToSignalPair.Value->Timestamp <= Now)
 		{
diff --git a/Engine/Plugins/VirtualProduction/DMX/DMXProtocol/Source/DMXProtocol/Public/IO/DMXOutputPort.h b/Engine/Plugins/VirtualProduction/DMX/DMXProtocol/Source/DMXProtocol/Public/IO/DMXOutputPort.h
--- a/Engine/Plugins/VirtualProduction/DMX/DMXProtocol/Source/DMXProtocol/Public/IO/DMXOutputPort.h
+++ b/Engine/Plugins/VirtualProduction/DMX/DMXProtocol/Source/DMXProtocol/Public/IO/DMXOutputPort.h
@@ -193,8 +193,11 @@ private:
 	/** Buffer of the signals that are to be sent in the next frame */
 	TQueue<TSharedPtr<FDMXSignalFragment, ESPMode::ThreadSafe>> SignalFragments;
 
-	/** Map that holds the latest Signal per Universe */
-	TMap<int32, FDMXSignalSharedPtr> ExternUniverseToLatestSignalMap;
+	/** Map that holds the latest Signal per Universe on the Game Thread */
+	TMap<int32, FDMXSignalSharedPtr> ExternUniverseToLatestSignalMap_GameThread;
+
+	/** Map that holds the latest Signal per Universe on the Port Thread */
+	TMap<int32, FDMXSignalSharedPtr> ExternUniverseToLatestSignalMap_PortThread;
 
 	/** The Destination Address to send to, can be irrelevant, e.g. for art-net broadcast */
 	TArray<FDMXOutputPortDestinationAddress> DestinationAddresses;
@@ -220,9 +223,6 @@ private:
 	/** The frame rate of the delay */
 	FFrameRate DelayFrameRate;
 
-	/** Critical section required to be used when the SignalFragments are accessed across threads */
-	FCriticalSection AccessExternUniverseToLatestSignalMapCriticalSection;
-
 	/** Critical section required to be used when clearing buffers */
 	FCriticalSection ClearBuffersCriticalSection;
 
diff --git a/Engine/Plugins/VirtualProduction/DMX/DMXProtocol/Source/DMXProtocol/Public/Interfaces/IDMXSender.h b/Engine/Plugins/VirtualProduction/DMX/DMXProtocol/Source/DMXProtocol/Public/Interfaces/IDMXSender.h
--- a/Engine/Plugins/VirtualProduction/DMX/DMXProtocol/Source/DMXProtocol/Public/Interfaces/IDMXSender.h
+++ b/Engine/Plugins/VirtualProduction/DMX/DMXProtocol/Source/DMXProtocol/Public/Interfaces/IDMXSender.h
@@ -16,7 +16,7 @@ public:
 	virtual ~IDMXSender()
 	{}
 
-	/** Sends the DMX signal. Note, this is called from a worker thread, not the game thread. */
+	/** Sends the DMX signal. Note, this may called concurrently from various threads */
 	virtual void SendDMXSignal(const FDMXSignalSharedRef& DMXSignal) = 0;
 
 	/** DEPRECATED 5.0 */
diff --git a/Engine/Plugins/VirtualProduction/DMX/DMXProtocol/Source/DMXProtocolArtNet/Private/DMXProtocolArtNetSender.cpp b/Engine/Plugins/VirtualProduction/DMX/DMXProtocol/Source/DMXProtocolArtNet/Private/DMXProtocolArtNetSender.cpp
--- a/Engine/Plugins/VirtualProduction/DMX/DMXProtocol/Source/DMXProtocolArtNet/Private/DMXProtocolArtNetSender.cpp
+++ b/Engine/Plugins/VirtualProduction/DMX/DMXProtocol/Source/DMXProtocolArtNet/Private/DMXProtocolArtNetSender.cpp
@@ -155,6 +155,9 @@ bool FDMXProtocolArtNetSender::IsCausingLoopback() const
 
 void FDMXProtocolArtNetSender::SendDMXSignal(const FDMXSignalSharedRef& DMXSignal)
 {
+	// This may be called from multiple threads, so calls need be synchronized
+	const FScopeLock SendDMXLock(&SendDMXCriticalSection);
+
 	FDMXProtocolArtNetDMXPacket ArtNetDMXPacket;
 	FMemory::Memcpy(ArtNetDMXPacket.Data, DMXSignal->ChannelData.GetData(), ARTNET_DMX_LENGTH);
 
diff --git a/Engine/Plugins/VirtualProduction/DMX/DMXProtocol/Source/DMXProtocolArtNet/Private/DMXProtocolArtNetSender.h b/Engine/Plugins/VirtualProduction/DMX/DMXProtocol/Source/DMXProtocolArtNet/Private/DMXProtocolArtNetSender.h
--- a/Engine/Plugins/VirtualProduction/DMX/DMXProtocol/Source/DMXProtocolArtNet/Private/DMXProtocolArtNetSender.h
+++ b/Engine/Plugins/VirtualProduction/DMX/DMXProtocol/Source/DMXProtocolArtNet/Private/DMXProtocolArtNetSender.h
@@ -7,6 +7,7 @@
 #include "Interfaces/IDMXSender.h"
 
 #include "CoreMinimal.h"
+#include "Misc/ScopeLock.h"
 #include "Serialization/ArrayWriter.h"
 
 enum class EDMXCommunicationType : uint8;
@@ -90,6 +91,9 @@ private:
 	/** Holds the network socket used to sender packages. */
 	FSocket* Socket = nullptr;
 
+	/** Critical section to be used during SendDMXSignal, as it may be called concurrently from many threads */
+	FCriticalSection SendDMXCriticalSection;
+
 	/** The network interface internet addr */
 	TSharedPtr<FInternetAddr> NetworkInterfaceInternetAddr;
 
diff --git a/Engine/Plugins/VirtualProduction/DMX/DMXProtocol/Source/DMXProtocolSACN/Private/DMXProtocolSACNSender.cpp b/Engine/Plugins/VirtualProduction/DMX/DMXProtocol/Source/DMXProtocolSACN/Private/DMXProtocolSACNSender.cpp
--- a/Engine/Plugins/VirtualProduction/DMX/DMXProtocol/Source/DMXProtocolSACN/Private/DMXProtocolSACNSender.cpp
+++ b/Engine/Plugins/VirtualProduction/DMX/DMXProtocol/Source/DMXProtocolSACN/Private/DMXProtocolSACNSender.cpp
@@ -166,6 +166,9 @@ bool FDMXProtocolSACNSender::IsCausingLoopback() const
 
 void FDMXProtocolSACNSender::SendDMXSignal(const FDMXSignalSharedRef& DMXSignal)
 {
+	// This may be called from multiple threads, so calls need be synchronized here
+	const FScopeLock SendDMXLock(&SendDMXCriticalSection);
+
 	TArray<uint8> Packet;
 
 	FDMXProtocolE131RootLayerPacket RootLayer;
@@ -176,9 +179,10 @@ void FDMXProtocolSACNSender::SendDMXSignal(const FDMXSignalSharedRef& DMXSignal)
 
 	int32 UniverseID = DMXSignal->ExternUniverseID;
 
+	check(UniverseID > 0 && UniverseID <= TNumericLimits<uint16>::Max());
 	FDMXProtocolE131FramingLayerPacket FramingLayer;
 	FramingLayer.Universe = UniverseID;
-	FramingLayer.SequenceNumber = UniverseIDToSequenceNumberMap.FindOrAdd(UniverseID, -1)++; // Init to max, let it wrap over to 0 at first
+	FramingLayer.SequenceNumber = UniverseIDToSequenceNumberMap.FindOrAdd(UniverseID, TNumericLimits<uint8>::Max())++; // Init to max, let it wrap over to 0 at first
 	FramingLayer.Priority = DMXSignal->Priority;
 	Packet.Append(*FramingLayer.Pack(ACN_DMX_SIZE));
 
diff --git a/Engine/Plugins/VirtualProduction/DMX/DMXProtocol/Source/DMXProtocolSACN/Private/DMXProtocolSACNSender.h b/Engine/Plugins/VirtualProduction/DMX/DMXProtocol/Source/DMXProtocolSACN/Private/DMXProtocolSACNSender.h
--- a/Engine/Plugins/VirtualProduction/DMX/DMXProtocol/Source/DMXProtocolSACN/Private/DMXProtocolSACNSender.h
+++ b/Engine/Plugins/VirtualProduction/DMX/DMXProtocol/Source/DMXProtocolSACN/Private/DMXProtocolSACNSender.h
@@ -7,9 +7,9 @@
 #include "Interfaces/IDMXSender.h"
 
 #include "CoreMinimal.h"
+#include "Misc/ScopeLock.h" 
 #include "Serialization/ArrayWriter.h"
 
-
 struct FDMXProtocolSACNRDM;
 
 enum class EDMXCommunicationType : uint8;
@@ -87,7 +87,7 @@ private:
 	TSet<TSharedPtr<FDMXOutputPort, ESPMode::ThreadSafe>> AssignedOutputPorts;
 
 	/** Map of universes with their current sequence number */
-	TMap<uint16 /** Universe ID */, uint16 /** Sequence Number */> UniverseIDToSequenceNumberMap;
+	TMap<int32 /** Universe ID */, uint8 /** Sequence Number */> UniverseIDToSequenceNumberMap;
 
 	/** The sACN protocol instance */
 	TSharedPtr<FDMXProtocolSACN, ESPMode::ThreadSafe> Protocol;
@@ -95,6 +95,9 @@ private:
 	/** Holds the network socket used to sender packages. */
 	FSocket* Socket = nullptr;
 
+	/** Critical section to be used during SendDMXSignal, as it may be called concurrently from many threads */
+	FCriticalSection SendDMXCriticalSection;
+
 	/** The network interface internet addr */
 	TSharedPtr<FInternetAddr> NetworkInterfaceInternetAddr;
 
