DMX - Fix InputPort function name for their InputDMXSignal method, express Mpsc implementation clearly for all

diff --git a/Engine/Plugins/VirtualProduction/DMX/DMXProtocol/Source/DMXProtocol/Private/IO/DMXInputPort.cpp b/Engine/Plugins/VirtualProduction/DMX/DMXProtocol/Source/DMXProtocol/Private/IO/DMXInputPort.cpp
--- a/Engine/Plugins/VirtualProduction/DMX/DMXProtocol/Source/DMXProtocol/Private/IO/DMXInputPort.cpp
+++ b/Engine/Plugins/VirtualProduction/DMX/DMXProtocol/Source/DMXProtocol/Private/IO/DMXInputPort.cpp
@@ -478,8 +478,30 @@ void FDMXInputPort::ClearBuffers()
 	ExternUniverseToLatestSignalMap.Reset();
 }
 
+void FDMXInputPort::InputDMXSignal(const FDMXSignalSharedRef& DMXSignal)
+{
+	if (IsReceiveDMXEnabled())
+	{
+		int32 ExternUniverseID = DMXSignal->ExternUniverseID;
+		if (IsExternUniverseInPortRange(ExternUniverseID))
+		{
+			for (const TSharedRef<FDMXRawListener>& RawListener : RawListeners)
+			{
+				RawListener->EnqueueSignal(this, DMXSignal);
+			}
+
+			if (bUseDefaultInputQueue)
+			{
+				DefaultInputQueue.Enqueue(DMXSignal);
+			}
+		}
+	}
+}
+
 void FDMXInputPort::SingleProducerInputDMXSignal(const FDMXSignalSharedRef& DMXSignal)
 {
+	// DERECATED 5.0
+
 	if (IsReceiveDMXEnabled())
 	{
 		int32 ExternUniverseID = DMXSignal->ExternUniverseID;
diff --git a/Engine/Plugins/VirtualProduction/DMX/DMXProtocol/Source/DMXProtocol/Public/IO/DMXInputPort.h b/Engine/Plugins/VirtualProduction/DMX/DMXProtocol/Source/DMXProtocol/Public/IO/DMXInputPort.h
--- a/Engine/Plugins/VirtualProduction/DMX/DMXProtocol/Source/DMXProtocol/Public/IO/DMXInputPort.h
+++ b/Engine/Plugins/VirtualProduction/DMX/DMXProtocol/Source/DMXProtocol/Public/IO/DMXInputPort.h
@@ -69,7 +69,11 @@ public:
 	/** Clears all buffers */
 	void ClearBuffers();
 
+	/** Thread-safe: Pushes a DMX Signal into the buffer */
+	void InputDMXSignal(const FDMXSignalSharedRef& DMXSignal);
+
 	/** Single Producer thread-safe: Pushes a DMX Signal into the buffer (For protocol only) */
+	UE_DEPRECATED(5.0, "Input ports now support multiple producers. Use InputDMXSignal instead.")
 	void SingleProducerInputDMXSignal(const FDMXSignalSharedRef& DMXSignal);
 
 	/** Gets the last signal received in specified local universe. Returns false if no signal was received. Game-Thread only */
@@ -104,7 +108,7 @@ private:
 	void OnSetReceiveDMXEnabled(bool bEnabled);
 
 	/** The default buffer, which is being read on tick */
-	TQueue<FDMXSignalSharedPtr> DefaultInputQueue;
+	TQueue<FDMXSignalSharedPtr, EQueueMode::Mpsc> DefaultInputQueue;
 
 	/** According to DMXProtcolSettings, true if DMX should be received */
 	bool bReceiveDMXEnabled = false;
diff --git a/Engine/Plugins/VirtualProduction/DMX/DMXProtocol/Source/DMXProtocolArtNet/Private/DMXProtocolArtNetReceiver.cpp b/Engine/Plugins/VirtualProduction/DMX/DMXProtocol/Source/DMXProtocolArtNet/Private/DMXProtocolArtNetReceiver.cpp
--- a/Engine/Plugins/VirtualProduction/DMX/DMXProtocol/Source/DMXProtocolArtNet/Private/DMXProtocolArtNetReceiver.cpp
+++ b/Engine/Plugins/VirtualProduction/DMX/DMXProtocol/Source/DMXProtocolArtNet/Private/DMXProtocolArtNetReceiver.cpp
@@ -236,7 +236,7 @@ void FDMXProtocolArtNetReceiver::HandleDataPacket(const TSharedRef<FArrayReader>
 		FDMXSignalSharedRef DMXSignal = MakeShared<FDMXSignal, ESPMode::ThreadSafe>(FPlatformTime::Seconds(), UniverseID, Priority, TArray<uint8>(ArtNetDMXPacket.Data, DMX_UNIVERSE_SIZE));
 		for (const TSharedPtr<FDMXInputPort, ESPMode::ThreadSafe>& InputPort : AssignedInputPorts)
 		{				
-			InputPort->SingleProducerInputDMXSignal(DMXSignal);
+			InputPort->InputDMXSignal(DMXSignal);
 		}
 
 		INC_DWORD_STAT(STAT_ArtNetPackagesReceived);
diff --git a/Engine/Plugins/VirtualProduction/DMX/DMXProtocol/Source/DMXProtocolSACN/Private/DMXProtocolSACNReceiver.cpp b/Engine/Plugins/VirtualProduction/DMX/DMXProtocol/Source/DMXProtocolSACN/Private/DMXProtocolSACNReceiver.cpp
--- a/Engine/Plugins/VirtualProduction/DMX/DMXProtocol/Source/DMXProtocolSACN/Private/DMXProtocolSACNReceiver.cpp
+++ b/Engine/Plugins/VirtualProduction/DMX/DMXProtocol/Source/DMXProtocolSACN/Private/DMXProtocolSACNReceiver.cpp
@@ -346,7 +346,7 @@ void FDMXProtocolSACNReceiver::HandleDataPacket(uint16 UniverseID, const TShared
 		{
 			continue;
 		}
-		InputPort->SingleProducerInputDMXSignal(DMXSignal);
+		InputPort->InputDMXSignal(DMXSignal);
 	}
 
 	INC_DWORD_STAT(STAT_SACNPackagesReceived);
