DMX - Fix protocols would not properly unregister ports in some cases, fix concurrency issues when modifying ports via console

diff --git a/Engine/Plugins/VirtualProduction/DMX/DMXProtocol/Source/DMXProtocol/Private/IO/DMXOutputPort.cpp b/Engine/Plugins/VirtualProduction/DMX/DMXProtocol/Source/DMXProtocol/Private/IO/DMXOutputPort.cpp
--- a/Engine/Plugins/VirtualProduction/DMX/DMXProtocol/Source/DMXProtocol/Private/IO/DMXOutputPort.cpp
+++ b/Engine/Plugins/VirtualProduction/DMX/DMXProtocol/Source/DMXProtocol/Private/IO/DMXOutputPort.cpp
@@ -56,7 +56,7 @@ static FAutoConsoleCommand GDMXSetOutputPortProtocolCommand(
 
 static FAutoConsoleCommand GDMXSetOutputPortCommunicationTypeCommand(
 	TEXT("DMX.SetOutputPortCommunicationType"),
-	TEXT("DMX.SetOutputPortCommunicationType [PortName][CommunicationType (1 = Broadcast, 2 = Multicast, 3 = Unicast)]. Sets the communication type of an output port. Example: DMX.SetOutputPortCommunicationType MyOutputPort 2"),
+	TEXT("DMX.SetOutputPortCommunicationType [PortName][CommunicationType (0 = Broadcast, 1 = Unicast, 2 = Multicast)]. Sets the communication type of an output port. Example: DMX.SetOutputPortCommunicationType MyOutputPort 2"),
 	FConsoleCommandWithArgsDelegate::CreateStatic(
 		[](const TArray<FString>& Args)
 		{
@@ -139,7 +139,7 @@ static FAutoConsoleCommand GDMXSetOutputPortOutputPortDestinationAddressesComman
 
 static FAutoConsoleCommand GDMXSetOutputPortLoopbackToEngineCommand(
 	TEXT("DMX.SetOutputPortInputIntoEngine"),
-	TEXT("DMX.SetOutputSetInputIntoEngine [PortName][Flag]. Sets if the Output Port is input into the engine directly. Example: DMX.SetOutputPortInputIntoEngine MyOutputPort 1"),
+	TEXT("DMX.SetOutputPortInputIntoEngine [PortName][Flag]. Sets if the Output Port is input into the engine directly. Example: DMX.SetOutputPortInputIntoEngine MyOutputPort 1"),
 	FConsoleCommandWithArgsDelegate::CreateStatic(
 		[](const TArray<FString>& Args)
 		{
@@ -695,6 +695,7 @@ bool FDMXOutputPort::Register()
 {
 	if (Protocol.IsValid() && IsValidPortSlow() && CommunicationDeterminator.IsSendDMXEnabled() && !FDMXPortManager::Get().AreProtocolsSuspended())
 	{
+		FScopeLock LockAccessSenderArray(&AccessSenderArrayCriticalSection);
 		DMXSenderArray = Protocol->RegisterOutputPort(SharedThis(this));
 
 		if (DMXSenderArray.Num() > 0)
@@ -717,6 +718,7 @@ void FDMXOutputPort::Unregister()
 			Protocol->UnregisterOutputPort(SharedThis(this));
 		}
 
+		FScopeLock LockAccessSenderArray(&AccessSenderArrayCriticalSection);
 		DMXSenderArray.Reset();
 	}
 
@@ -852,6 +854,7 @@ void FDMXOutputPort::ProcessSendDMX()
 			// Send via the protocol's sender
 			if (bNeedsSendDMX)
 			{
+				FScopeLock LockAccessSenderArray(&AccessSenderArrayCriticalSection);
 				for (const TSharedPtr<IDMXSender>& DMXSender : DMXSenderArray)
 				{
 					DMXSender->SendDMXSignal(UniverseToSignalPair.Value.ToSharedRef());
diff --git a/Engine/Plugins/VirtualProduction/DMX/DMXProtocol/Source/DMXProtocol/Public/IO/DMXOutputPort.h b/Engine/Plugins/VirtualProduction/DMX/DMXProtocol/Source/DMXProtocol/Public/IO/DMXOutputPort.h
--- a/Engine/Plugins/VirtualProduction/DMX/DMXProtocol/Source/DMXProtocol/Public/IO/DMXOutputPort.h
+++ b/Engine/Plugins/VirtualProduction/DMX/DMXProtocol/Source/DMXProtocol/Public/IO/DMXOutputPort.h
@@ -223,6 +223,9 @@ private:
 	/** The frame rate of the delay */
 	FFrameRate DelayFrameRate;
 
+	/** Critical section required to be used when clearing buffers */
+	FCriticalSection AccessSenderArrayCriticalSection;
+
 	/** Critical section required to be used when clearing buffers */
 	FCriticalSection ClearBuffersCriticalSection;
 
diff --git a/Engine/Plugins/VirtualProduction/DMX/DMXProtocol/Source/DMXProtocolArtNet/Private/DMXProtocolArtNet.cpp b/Engine/Plugins/VirtualProduction/DMX/DMXProtocol/Source/DMXProtocolArtNet/Private/DMXProtocolArtNet.cpp
--- a/Engine/Plugins/VirtualProduction/DMX/DMXProtocol/Source/DMXProtocolArtNet/Private/DMXProtocolArtNet.cpp
+++ b/Engine/Plugins/VirtualProduction/DMX/DMXProtocol/Source/DMXProtocolArtNet/Private/DMXProtocolArtNet.cpp
@@ -125,8 +125,7 @@ void FDMXProtocolArtNet::UnregisterInputPort(const FDMXInputPortSharedRef& Input
 	check(CachedInputPorts.Contains(InputPort));
 	CachedInputPorts.Remove(InputPort);
 
-	TSharedPtr<FDMXProtocolArtNetReceiver> UnusedReceiver;
-	for (const TSharedPtr<FDMXProtocolArtNetReceiver>& Receiver : Receivers)
+	for (const TSharedPtr<FDMXProtocolArtNetReceiver>& Receiver : TSet<TSharedPtr<FDMXProtocolArtNetReceiver>>(Receivers))
 	{
 		if (Receiver->ContainsInputPort(InputPort))
 		{
@@ -134,16 +133,10 @@ void FDMXProtocolArtNet::UnregisterInputPort(const FDMXInputPortSharedRef& Input
 
 			if (Receiver->GetNumAssignedInputPorts() == 0)
 			{
-				UnusedReceiver = Receiver;
+				Receivers.Remove(Receiver);
 			}
-			break;
 		}
 	}
-
-	if (UnusedReceiver.IsValid())
-	{
-		Receivers.Remove(UnusedReceiver);
-	}
 }
 
 TArray<TSharedPtr<IDMXSender>> FDMXProtocolArtNet::RegisterOutputPort(const FDMXOutputPortSharedRef& OutputPort)
@@ -214,7 +207,6 @@ void FDMXProtocolArtNet::UnregisterOutputPort(const FDMXOutputPortSharedRef& Out
 	check(CachedOutputPorts.Contains(OutputPort));
 	CachedOutputPorts.Remove(OutputPort);
 
-	TArray<TSharedPtr<FDMXProtocolArtNetSender>> UnusedSenders;
 	for (const TSharedPtr<FDMXProtocolArtNetSender>& Sender : TSet<TSharedPtr<FDMXProtocolArtNetSender>>(Senders))
 	{
 		if (Sender->ContainsOutputPort(OutputPort))
diff --git a/Engine/Plugins/VirtualProduction/DMX/DMXProtocol/Source/DMXProtocolSACN/Private/DMXProtocolSACN.cpp b/Engine/Plugins/VirtualProduction/DMX/DMXProtocol/Source/DMXProtocolSACN/Private/DMXProtocolSACN.cpp
--- a/Engine/Plugins/VirtualProduction/DMX/DMXProtocol/Source/DMXProtocolSACN/Private/DMXProtocolSACN.cpp
+++ b/Engine/Plugins/VirtualProduction/DMX/DMXProtocol/Source/DMXProtocolSACN/Private/DMXProtocolSACN.cpp
@@ -127,8 +127,7 @@ void FDMXProtocolSACN::UnregisterInputPort(const TSharedRef<FDMXInputPort, ESPMo
 	check(CachedInputPorts.Contains(InputPort));
 	CachedInputPorts.Remove(InputPort);
 
-	TSharedPtr<FDMXProtocolSACNReceiver> UnusedReceiver;
-	for (const TSharedPtr<FDMXProtocolSACNReceiver>& Receiver : Receivers)
+	for (const TSharedPtr<FDMXProtocolSACNReceiver>& Receiver : TSet<TSharedPtr<FDMXProtocolSACNReceiver>>(Receivers))
 	{
 		if (Receiver->ContainsInputPort(InputPort))
 		{
@@ -136,16 +135,10 @@ void FDMXProtocolSACN::UnregisterInputPort(const TSharedRef<FDMXInputPort, ESPMo
 
 			if (Receiver->GetNumAssignedInputPorts() == 0)
 			{
-				UnusedReceiver = Receiver;
+				Receivers.Remove(Receiver);
 			}
-			break;
 		}
 	}
-
-	if (UnusedReceiver.IsValid())
-	{
-		Receivers.Remove(UnusedReceiver);
-	}
 }
 
 TArray<TSharedPtr<IDMXSender>> FDMXProtocolSACN::RegisterOutputPort(const TSharedRef<FDMXOutputPort, ESPMode::ThreadSafe>& OutputPort)
@@ -157,7 +150,7 @@ TArray<TSharedPtr<IDMXSender>> FDMXProtocolSACN::RegisterOutputPort(const TShare
 	EDMXCommunicationType CommunicationType = OutputPort->GetCommunicationType();
 
 	// Try to use an existing receiver or create a new one
-	TArray<TSharedPtr<IDMXSender>> NewsSenders;
+	TArray<TSharedPtr<IDMXSender>> NewSenders;
 	
 	if (CommunicationType == EDMXCommunicationType::Multicast)
 	{
@@ -169,10 +162,15 @@ TArray<TSharedPtr<IDMXSender>> FDMXProtocolSACN::RegisterOutputPort(const TShare
 		}
 
 		NewSender->AssignOutputPort(OutputPort);
-		NewsSenders.Add(NewSender);
+		NewSenders.Add(NewSender);
 	}
 	else if (CommunicationType == EDMXCommunicationType::Unicast)
 	{
+		if (OutputPort->GetDestinationAddresses().Num() == 0)
+		{
+			UE_LOG(LogDMXProtocol, Warning, TEXT("Cannot create DMX Protocol sACN Sender for Output Port '%s'. Port is set to Unicast, but does not specify any Destination Addresses."), *OutputPort->GetPortName());
+		}
+
 		for (const FString& UnicastAddress : OutputPort->GetDestinationAddresses())
 		{
 			TSharedPtr<FDMXProtocolSACNSender> NewSender = FindExistingUnicastSender(NetworkInterfaceAddress, UnicastAddress);
@@ -183,23 +181,25 @@ TArray<TSharedPtr<IDMXSender>> FDMXProtocolSACN::RegisterOutputPort(const TShare
 			}
 
 			NewSender->AssignOutputPort(OutputPort);
-			NewsSenders.Add(NewSender);
+			NewSenders.Add(NewSender);
 		}
 	}
 	else
 	{
 		// Invalid Communication Type
-		UE_LOG(LogDMXProtocol, Error, TEXT("Cannot create DMX Protocol sACN Sender. The communication type specified is not supported."));
+		UE_LOG(LogDMXProtocol, Warning, TEXT("Cannot create DMX Protocol sACN Sender. The communication type specified is not supported."));
 	}
 
-	if (NewsSenders.Num() == 0)
+	if (NewSenders.Num() > 0)
 	{
-		UE_LOG(LogDMXProtocol, Warning, TEXT("Could not create sACN sender for output port %s"), *OutputPort->GetPortName());
+		CachedOutputPorts.Add(OutputPort);
+	}
+	else
+	{
+		UE_LOG(LogDMXProtocol, Warning, TEXT("Could not create sACN sender for Output Port '%s'"), *OutputPort->GetPortName());
 	}
 
-	CachedOutputPorts.Add(OutputPort);
-
-	return NewsSenders;
+	return NewSenders;
 }
 
 void FDMXProtocolSACN::UnregisterOutputPort(const TSharedRef<FDMXOutputPort, ESPMode::ThreadSafe>& OutputPort)
@@ -207,7 +207,6 @@ void FDMXProtocolSACN::UnregisterOutputPort(const TSharedRef<FDMXOutputPort, ESP
 	check(CachedOutputPorts.Contains(OutputPort));
 	CachedOutputPorts.Remove(OutputPort);
 
-	TSharedPtr<FDMXProtocolSACNSender> UnusedSender;
 	for (const TSharedPtr<FDMXProtocolSACNSender>& Sender : TSet<TSharedPtr<FDMXProtocolSACNSender>>(Senders))
 	{
 		if (Sender->ContainsOutputPort(OutputPort))
