DMX - Fix an issue where Input and Output Port References would set their default in their graph pin instead of the struct self

diff --git a/Engine/Plugins/VirtualProduction/DMX/DMXProtocol/Source/DMXProtocol/Public/IO/DMXInputPortReference.h b/Engine/Plugins/VirtualProduction/DMX/DMXProtocol/Source/DMXProtocol/Public/IO/DMXInputPortReference.h
--- a/Engine/Plugins/VirtualProduction/DMX/DMXProtocol/Source/DMXProtocol/Public/IO/DMXInputPortReference.h
+++ b/Engine/Plugins/VirtualProduction/DMX/DMXProtocol/Source/DMXProtocol/Public/IO/DMXInputPortReference.h
@@ -2,6 +2,9 @@
 
 #pragma once
 
+#include "DMXProtocolLog.h"
+#include "DMXProtocolSettings.h"
+
 #include "Misc/Guid.h"
 
 #include "DMXInputPortReference.generated.h"
@@ -15,17 +18,45 @@ struct DMXPROTOCOL_API FDMXInputPortReference
 
 	FDMXInputPortReference()
 		: bEnabledFlag(true)
-	{}
+	{
+		const UDMXProtocolSettings* ProtocolSettings = GetDefault<UDMXProtocolSettings>();
+		if (ProtocolSettings && ProtocolSettings->InputPortConfigs.Num() > 0)
+		{
+			PortGuid = ProtocolSettings->InputPortConfigs[0].GetPortGuid();
+		}
+	}
 
 	FDMXInputPortReference(const FGuid& InPortGuid, bool bIsEnabledFlag)
 		: PortGuid(InPortGuid)
 		, bEnabledFlag(bIsEnabledFlag)
-	{}
+	{
+		const UDMXProtocolSettings* ProtocolSettings = GetDefault<UDMXProtocolSettings>();
+		if (ProtocolSettings)
+		{
+			const FDMXInputPortConfig* ExistingInputPortConfigPtr = ProtocolSettings->InputPortConfigs.FindByPredicate([InPortGuid](const FDMXInputPortConfig& InputPortConfig)
+				{
+					return InputPortConfig.GetPortGuid() == InPortGuid;
+				});
+
+			UE_CLOG(!ExistingInputPortConfigPtr, LogDMXProtocol, Warning, TEXT("Trying to construct input port refrence, but the referenced Port Guid doesn't exist."));
+		}
+	}
 
 	FDMXInputPortReference(const FDMXInputPortReference& InInputPortReference, bool bIsEnabledFlag)
 		: PortGuid(InInputPortReference.PortGuid)
 		, bEnabledFlag(bIsEnabledFlag)
-	{}
+	{
+		const UDMXProtocolSettings* ProtocolSettings = GetDefault<UDMXProtocolSettings>();
+		if (ProtocolSettings)
+		{
+			const FDMXInputPortConfig* ExistingInputPortConfigPtr = ProtocolSettings->InputPortConfigs.FindByPredicate([InInputPortReference](const FDMXInputPortConfig& InputPortConfig)
+				{
+					return InputPortConfig.GetPortGuid() == InInputPortReference.GetPortGuid();
+				});
+
+			UE_CLOG(!ExistingInputPortConfigPtr, LogDMXProtocol, Warning, TEXT("Trying to construct input port refrence, but the referenced Port Guid doesn't exist."));
+		}
+	}
 
 	/** Returns true if the port is enabled. Always true unless constructed with bIsAlwaysEnabled = false */
 	FORCEINLINE bool IsEnabledFlagSet() const { return bEnabledFlag; }
diff --git a/Engine/Plugins/VirtualProduction/DMX/DMXProtocol/Source/DMXProtocol/Public/IO/DMXOutputPortReference.h b/Engine/Plugins/VirtualProduction/DMX/DMXProtocol/Source/DMXProtocol/Public/IO/DMXOutputPortReference.h
--- a/Engine/Plugins/VirtualProduction/DMX/DMXProtocol/Source/DMXProtocol/Public/IO/DMXOutputPortReference.h
+++ b/Engine/Plugins/VirtualProduction/DMX/DMXProtocol/Source/DMXProtocol/Public/IO/DMXOutputPortReference.h
@@ -2,6 +2,9 @@
 
 #pragma once
 
+#include "DMXProtocolLog.h"
+#include "DMXProtocolSettings.h"
+
 #include "Misc/Guid.h"
 
 #include "DMXOutputPortReference.generated.h"
@@ -15,17 +18,45 @@ struct DMXPROTOCOL_API FDMXOutputPortReference
 
 	FDMXOutputPortReference()
 		: bEnabledFlag(true)
-	{}
+	{
+		const UDMXProtocolSettings* ProtocolSettings = GetDefault<UDMXProtocolSettings>();
+		if (ProtocolSettings && ProtocolSettings->OutputPortConfigs.Num() > 0)
+		{
+			PortGuid = ProtocolSettings->OutputPortConfigs[0].GetPortGuid();
+		}
+	}
 
 	FDMXOutputPortReference(const FGuid& InPortGuid, bool bIsEnabledFlag)
 		: PortGuid(InPortGuid)
 		, bEnabledFlag(bIsEnabledFlag)
-	{}
+	{
+		const UDMXProtocolSettings* ProtocolSettings = GetDefault<UDMXProtocolSettings>();
+		if (ProtocolSettings)
+		{
+			const FDMXOutputPortConfig* ExistingOutputPortConfigPtr = ProtocolSettings->OutputPortConfigs.FindByPredicate([InPortGuid](const FDMXOutputPortConfig& OutputPortConfig)
+				{
+					return OutputPortConfig.GetPortGuid() == InPortGuid;
+				});
+
+			UE_CLOG(!ExistingOutputPortConfigPtr, LogDMXProtocol, Warning, TEXT("Trying to construct output port refrence, but the referenced Port Guid doesn't exist."));
+		}
+	}
 
 	FDMXOutputPortReference(const FDMXOutputPortReference& InOutputPortReference, bool bIsEnabledFlag)
 		: PortGuid(InOutputPortReference.PortGuid)
 		, bEnabledFlag(bIsEnabledFlag)
-	{}
+	{
+		const UDMXProtocolSettings* ProtocolSettings = GetDefault<UDMXProtocolSettings>();
+		if (ProtocolSettings)
+		{
+			const FDMXOutputPortConfig* ExistingOutputPortConfigPtr = ProtocolSettings->OutputPortConfigs.FindByPredicate([InOutputPortReference](const FDMXOutputPortConfig& OutputPortConfig)
+				{
+					return OutputPortConfig.GetPortGuid() == InOutputPortReference.GetPortGuid();
+				});
+
+			UE_CLOG(!ExistingOutputPortConfigPtr, LogDMXProtocol, Warning, TEXT("Trying to construct output port refrence, but the referenced Port Guid doesn't exist."));
+		}
+	}
 
 	/** Returns true if the port is enabled. Always true unless constructed with bIsAlwaysEnabled = false */
 	FORCEINLINE bool IsEnabledFlagSet() const { return bEnabledFlag; }
diff --git a/Engine/Plugins/VirtualProduction/DMX/DMXProtocol/Source/DMXProtocolBlueprintGraph/Private/Widgets/SDMXInputPortReferenceGraphPin.cpp b/Engine/Plugins/VirtualProduction/DMX/DMXProtocol/Source/DMXProtocolBlueprintGraph/Private/Widgets/SDMXInputPortReferenceGraphPin.cpp
--- a/Engine/Plugins/VirtualProduction/DMX/DMXProtocol/Source/DMXProtocolBlueprintGraph/Private/Widgets/SDMXInputPortReferenceGraphPin.cpp
+++ b/Engine/Plugins/VirtualProduction/DMX/DMXProtocol/Source/DMXProtocolBlueprintGraph/Private/Widgets/SDMXInputPortReferenceGraphPin.cpp
@@ -21,7 +21,6 @@ TSharedRef<SWidget> SDMXInputPortReferenceGraphPin::GetDefaultValueWidget()
 {
 	FDMXInputPortReference InitiallySelectedPortReference = GetPinValue();
 
-	// Create entity picker button
 	PortSelector = SNew(SDMXPortSelector)
 		.Mode(EDMXPortSelectorMode::SelectFromAvailableInputs)
 		.InitialSelection(InitiallySelectedPortReference.GetPortGuid())
@@ -41,29 +40,19 @@ FDMXInputPortReference SDMXInputPortReferenceGraphPin::GetPinValue() const
 		FDMXInputPortReference::StaticStruct()->ImportText(*EntityRefStr, &PortReference, nullptr, EPropertyPortFlags::PPF_None, GLog, FDMXInputPortReference::StaticStruct()->GetName());
 	}
 
-	// If no value is set or the value doesn't exist in protocol settings, return the first port in settings instead.
-	if (!PortReference.GetPortGuid().IsValid())
-	{
-		const UDMXProtocolSettings* ProtocolSettings = GetDefault<UDMXProtocolSettings>();
-		if (ensure(ProtocolSettings))
-		{
-			if (ProtocolSettings->InputPortConfigs.Num() > 0)
-			{
-				PortReference = FDMXInputPortReference(ProtocolSettings->InputPortConfigs[0].GetPortGuid(), true);
-
-				constexpr bool bMarkAsModified = false;
-				SetPinValue(PortReference, bMarkAsModified);
-			}
-		}
-	}
-
 	return PortReference;
 }
-	
+
 void SDMXInputPortReferenceGraphPin::SetPinValue(const FDMXInputPortReference& InputPortReference, bool bMarkAsModified) const
 {
+	if (GraphPinObj->IsPendingKill())
+	{
+		return;
+	}
+
 	FString ValueString;
 	FDMXInputPortReference::StaticStruct()->ExportText(ValueString, &InputPortReference, nullptr, nullptr, EPropertyPortFlags::PPF_None, nullptr);
+
 	GraphPinObj->GetSchema()->TrySetDefaultValue(*GraphPinObj, ValueString, bMarkAsModified);
 }
 
@@ -90,7 +79,8 @@ void SDMXInputPortReferenceGraphPin::OnPortSelected() const
 
 		FDMXInputPortReference PortReference(PortGuid, true);
 
-		SetPinValue(PortReference);
+		constexpr bool bMarkAsModified = true;
+		SetPinValue(PortReference, bMarkAsModified);
 	}
 }
 
diff --git a/Engine/Plugins/VirtualProduction/DMX/DMXProtocol/Source/DMXProtocolBlueprintGraph/Private/Widgets/SDMXInputPortReferenceGraphPin.h b/Engine/Plugins/VirtualProduction/DMX/DMXProtocol/Source/DMXProtocolBlueprintGraph/Private/Widgets/SDMXInputPortReferenceGraphPin.h
--- a/Engine/Plugins/VirtualProduction/DMX/DMXProtocol/Source/DMXProtocolBlueprintGraph/Private/Widgets/SDMXInputPortReferenceGraphPin.h
+++ b/Engine/Plugins/VirtualProduction/DMX/DMXProtocol/Source/DMXProtocolBlueprintGraph/Private/Widgets/SDMXInputPortReferenceGraphPin.h
@@ -31,7 +31,7 @@ private:
 
 	FDMXInputPortReference GetPinValue() const;
 
-	void SetPinValue(const FDMXInputPortReference& InputPortReference, bool bMarkAsModified = true) const;
+	void SetPinValue(const FDMXInputPortReference& InputPortReference, bool bMarkAsModified) const;
 
 	TSharedPtr<SDMXPortSelector> PortSelector;
 };
diff --git a/Engine/Plugins/VirtualProduction/DMX/DMXProtocol/Source/DMXProtocolBlueprintGraph/Private/Widgets/SDMXOutputPortReferenceGraphPin.cpp b/Engine/Plugins/VirtualProduction/DMX/DMXProtocol/Source/DMXProtocolBlueprintGraph/Private/Widgets/SDMXOutputPortReferenceGraphPin.cpp
--- a/Engine/Plugins/VirtualProduction/DMX/DMXProtocol/Source/DMXProtocolBlueprintGraph/Private/Widgets/SDMXOutputPortReferenceGraphPin.cpp
+++ b/Engine/Plugins/VirtualProduction/DMX/DMXProtocol/Source/DMXProtocolBlueprintGraph/Private/Widgets/SDMXOutputPortReferenceGraphPin.cpp
@@ -40,29 +40,19 @@ FDMXOutputPortReference SDMXOutputPortReferenceGraphPin::GetPinValue() const
 		FDMXOutputPortReference::StaticStruct()->ImportText(*EntityRefStr, &PortReference, nullptr, EPropertyPortFlags::PPF_None, GLog, FDMXOutputPortReference::StaticStruct()->GetName());
 	}
 
-	// If no value is set or the value doesn't exist in protocol settings, return the first port in settings instead.
-	if (!PortReference.GetPortGuid().IsValid())
-	{
-		const UDMXProtocolSettings* ProtocolSettings = GetDefault<UDMXProtocolSettings>();
-		if (ensure(ProtocolSettings))
-		{
-			if (ProtocolSettings->OutputPortConfigs.Num() > 0)
-			{
-				PortReference = FDMXOutputPortReference(ProtocolSettings->OutputPortConfigs[0].GetPortGuid(), true);
-
-				constexpr bool bMarkAsModified = false;
-				SetPinValue(PortReference, bMarkAsModified);
-			}
-		}
-	}
-
 	return PortReference;
 }
 
 void SDMXOutputPortReferenceGraphPin::SetPinValue(const FDMXOutputPortReference& OutputPortReference, bool bMarkAsModified) const
 {
+	if (GraphPinObj->IsPendingKill())
+	{
+		return;
+	}
+
 	FString ValueString;
 	FDMXOutputPortReference::StaticStruct()->ExportText(ValueString, &OutputPortReference, nullptr, nullptr, EPropertyPortFlags::PPF_None, nullptr);
+
 	GraphPinObj->GetSchema()->TrySetDefaultValue(*GraphPinObj, ValueString, bMarkAsModified);
 }
 
