DMX - Fix an issue where sACN universe changes were only applied after an editor restart, fix issues where some console commands would not properly update ports at runtime

diff --git a/Engine/Plugins/VirtualProduction/DMX/DMXProtocol/Source/DMXProtocol/Private/DMXProtocolSettings.cpp b/Engine/Plugins/VirtualProduction/DMX/DMXProtocol/Source/DMXProtocol/Private/DMXProtocolSettings.cpp
--- a/Engine/Plugins/VirtualProduction/DMX/DMXProtocol/Source/DMXProtocol/Private/DMXProtocolSettings.cpp
+++ b/Engine/Plugins/VirtualProduction/DMX/DMXProtocol/Source/DMXProtocol/Private/DMXProtocolSettings.cpp
@@ -185,53 +185,56 @@ void UDMXProtocolSettings::PostEditChangeChainProperty(FPropertyChangedChainEven
 
 		FDMXAttributeName::OnValuesChanged.Broadcast();
 	}
-	else if	(
-		PropertyName == GET_MEMBER_NAME_CHECKED(UDMXProtocolSettings, InputPortConfigs) ||
-		PropertyName == GET_MEMBER_NAME_CHECKED(UDMXProtocolSettings, OutputPortConfigs) ||
-		(InputPortConfigStruct && InputPortConfigStruct == PropertyOwnerStruct) ||
-		(OutputPortConfigStruct && OutputPortConfigStruct == PropertyOwnerStruct))
+	else if (
+		PropertyName == FDMXOutputPortConfig::GetDestinationAddressesPropertyNameChecked() ||
+		PropertyName == GET_MEMBER_NAME_CHECKED(FDMXOutputPortDestinationAddress, DestinationAddressString))
+	{
+		FDMXPortManager::Get().UpdateFromProtocolSettings();
+	}
+	else if (PropertyChangedChainEvent.ChangeType != EPropertyChangeType::Interactive)
 	{
-		if (PropertyChangedChainEvent.ChangeType == EPropertyChangeType::Duplicate)
+		if (PropertyName == GET_MEMBER_NAME_CHECKED(UDMXProtocolSettings, InputPortConfigs) ||
+			PropertyName == GET_MEMBER_NAME_CHECKED(UDMXProtocolSettings, OutputPortConfigs) ||
+			(InputPortConfigStruct && InputPortConfigStruct == PropertyOwnerStruct) ||
+			(OutputPortConfigStruct && OutputPortConfigStruct == PropertyOwnerStruct))
 		{
-			// When duplicating configs, the guid will be duplicated, so we have to create unique ones instead
-
-			int32 ChangedIndex = PropertyChangedChainEvent.GetArrayIndex(PropertyName.ToString());
-			if (ensureAlways(ChangedIndex != INDEX_NONE))
+			if (PropertyChangedChainEvent.ChangeType == EPropertyChangeType::Duplicate)
 			{
-				if (PropertyName == GET_MEMBER_NAME_CHECKED(UDMXProtocolSettings, InputPortConfigs))
-				{
-					const int32 IndexOfDuplicate = InputPortConfigs.FindLastByPredicate([this, ChangedIndex](const FDMXInputPortConfig& InputPortConfig) {
-						return InputPortConfigs[ChangedIndex].GetPortGuid() == InputPortConfig.GetPortGuid();
-						});
+				// When duplicating configs, the guid will be duplicated, so we have to create unique ones instead
 
-					if (ensureAlways(IndexOfDuplicate != ChangedIndex))
+				int32 ChangedIndex = PropertyChangedChainEvent.GetArrayIndex(PropertyName.ToString());
+				if (ensureAlways(ChangedIndex != INDEX_NONE))
+				{
+					if (PropertyName == GET_MEMBER_NAME_CHECKED(UDMXProtocolSettings, InputPortConfigs))
 					{
-						InputPortConfigs[IndexOfDuplicate] = FDMXInputPortConfig(FGuid::NewGuid(), InputPortConfigs[ChangedIndex]);
+						const int32 IndexOfDuplicate = InputPortConfigs.FindLastByPredicate([this, ChangedIndex](const FDMXInputPortConfig& InputPortConfig) {
+							return InputPortConfigs[ChangedIndex].GetPortGuid() == InputPortConfig.GetPortGuid();
+							});
+
+						if (ensureAlways(IndexOfDuplicate != ChangedIndex))
+						{
+							InputPortConfigs[IndexOfDuplicate] = FDMXInputPortConfig(FGuid::NewGuid(), InputPortConfigs[ChangedIndex]);
+						}
 					}
-				}
-				else if (PropertyName == GET_MEMBER_NAME_CHECKED(UDMXProtocolSettings, OutputPortConfigs))
-				{
-					const int32 IndexOfDuplicate = OutputPortConfigs.FindLastByPredicate([this, ChangedIndex](const FDMXOutputPortConfig& OutputPortConfig) {
-						return OutputPortConfigs[ChangedIndex].GetPortGuid() == OutputPortConfig.GetPortGuid();
-						});
-
-					if (ensureAlways(IndexOfDuplicate != ChangedIndex))
+					else if (PropertyName == GET_MEMBER_NAME_CHECKED(UDMXProtocolSettings, OutputPortConfigs))
 					{
-						OutputPortConfigs[IndexOfDuplicate] = FDMXOutputPortConfig(FGuid::NewGuid(), OutputPortConfigs[ChangedIndex]);
+						const int32 IndexOfDuplicate = OutputPortConfigs.FindLastByPredicate([this, ChangedIndex](const FDMXOutputPortConfig& OutputPortConfig) {
+							return OutputPortConfigs[ChangedIndex].GetPortGuid() == OutputPortConfig.GetPortGuid();
+							});
+
+						if (ensureAlways(IndexOfDuplicate != ChangedIndex))
+						{
+							OutputPortConfigs[IndexOfDuplicate] = FDMXOutputPortConfig(FGuid::NewGuid(), OutputPortConfigs[ChangedIndex]);
+						}
 					}
 				}
 			}
-		}
 
-		FDMXPortManager::Get().UpdateFromProtocolSettings();
-	}
-	else if (
-		PropertyName == FDMXOutputPortConfig::GetDestinationAddressesPropertyNameChecked() ||
-		PropertyName == GET_MEMBER_NAME_CHECKED(FDMXOutputPortDestinationAddress, DestinationAddressString))
-	{
-		FDMXPortManager::Get().UpdateFromProtocolSettings();
+			constexpr bool bForceUpdateRegistrationWithProtocol = true;
+			FDMXPortManager::Get().UpdateFromProtocolSettings(bForceUpdateRegistrationWithProtocol);
+		}
 	}
-	
+
 	Super::PostEditChangeChainProperty(PropertyChangedChainEvent);
 }
 #endif // WITH_EDITOR
diff --git a/Engine/Plugins/VirtualProduction/DMX/DMXProtocol/Source/DMXProtocol/Private/IO/DMXInputPort.cpp b/Engine/Plugins/VirtualProduction/DMX/DMXProtocol/Source/DMXProtocol/Private/IO/DMXInputPort.cpp
--- a/Engine/Plugins/VirtualProduction/DMX/DMXProtocol/Source/DMXProtocol/Private/IO/DMXInputPort.cpp
+++ b/Engine/Plugins/VirtualProduction/DMX/DMXProtocol/Source/DMXProtocol/Private/IO/DMXInputPort.cpp
@@ -29,7 +29,8 @@ DECLARE_CYCLE_STAT(TEXT("Input Port Tick"), STAT_DMXInputPortTick, STATGROUP_DMX
 		FDMXInputPortConfigParams InputPortConfigParams = FDMXInputPortConfigParams(OldInputPortConfig); \
 		InputPortConfigParams.MemberName = Value; \
 		FDMXInputPortConfig NewInputPortConfig((*InputPortPtr)->GetPortGuid(), InputPortConfigParams); \
-		(*InputPortPtr)->UpdateFromConfig(NewInputPortConfig); \
+		constexpr bool bForceUpdateRegistrationWithProtocol = true; \
+		(*InputPortPtr)->UpdateFromConfig(NewInputPortConfig, bForceUpdateRegistrationWithProtocol); \
 	} \
 }
 
@@ -282,14 +283,19 @@ FDMXInputPortConfig FDMXInputPort::MakeInputPortConfig() const
 	Params.LocalUniverseStart = LocalUniverseStart;
 	Params.NumUniverses = NumUniverses;
 	Params.ExternUniverseStart = ExternUniverseStart;
+	Params.PriorityStrategy = PriorityStrategy;
+	Params.Priority = Priority;
 
 	return FDMXInputPortConfig(PortGuid, Params);
 }
 
-void FDMXInputPort::UpdateFromConfig(FDMXInputPortConfig& InputPortConfig)
+void FDMXInputPort::UpdateFromConfig(FDMXInputPortConfig& InOutInputPortConfig, bool bForceUpdateRegistrationWithProtocol)
 {
 	// Need a valid config for the port
-	InputPortConfig.MakeValid();
+	InOutInputPortConfig.MakeValid();
+
+	// Avoid further changes to the config
+	const FDMXInputPortConfig& InputPortConfig = InOutInputPortConfig;
 
 	// Can only use configs that correspond to project settings
 	const UDMXProtocolSettings* ProtocolSettings = GetDefault<UDMXProtocolSettings>();
@@ -299,8 +305,13 @@ void FDMXInputPort::UpdateFromConfig(FDMXInputPortConfig& InputPortConfig)
 	checkf(bConfigIsInProjectSettings, TEXT("Can only use configs with a guid that corresponds to a config in project settings"));
 
 	// Find if the port needs update its registration with the protocol
-	const bool bNeedsUpdateRegistration = [this, &InputPortConfig]()
+	const bool bNeedsUpdateRegistration = [this, &InputPortConfig, bForceUpdateRegistrationWithProtocol]()
 	{
+		if (bForceUpdateRegistrationWithProtocol)
+		{
+			return true;
+		}
+
 		if (IsRegistered() != bReceiveDMXEnabled)
 		{
 			return true;
@@ -348,6 +359,7 @@ void FDMXInputPort::UpdateFromConfig(FDMXInputPortConfig& InputPortConfig)
 		Register();
 	}
 
+
 	OnPortUpdated.Broadcast();
 }
 
diff --git a/Engine/Plugins/VirtualProduction/DMX/DMXProtocol/Source/DMXProtocol/Private/IO/DMXOutputPort.cpp b/Engine/Plugins/VirtualProduction/DMX/DMXProtocol/Source/DMXProtocol/Private/IO/DMXOutputPort.cpp
--- a/Engine/Plugins/VirtualProduction/DMX/DMXProtocol/Source/DMXProtocol/Private/IO/DMXOutputPort.cpp
+++ b/Engine/Plugins/VirtualProduction/DMX/DMXProtocol/Source/DMXProtocol/Private/IO/DMXOutputPort.cpp
@@ -29,7 +29,8 @@
 		FDMXOutputPortConfigParams OutputPortConfigParams = FDMXOutputPortConfigParams(OldOutputPortConfig); \
 		OutputPortConfigParams.MemberName = Value; \
 		FDMXOutputPortConfig NewOutputPortConfig((*OutputPortPtr)->GetPortGuid(), OutputPortConfigParams); \
-		(*OutputPortPtr)->UpdateFromConfig(NewOutputPortConfig); \
+		constexpr bool bForceUpdateRegistrationWithProtocol = true; \
+		(*OutputPortPtr)->UpdateFromConfig(NewOutputPortConfig, bForceUpdateRegistrationWithProtocol); \
 	} \
 }
 
@@ -426,10 +427,13 @@ FDMXOutputPortConfig FDMXOutputPort::MakeOutputPortConfig() const
 	return FDMXOutputPortConfig(PortGuid, Params);
 }
 
-void FDMXOutputPort::UpdateFromConfig(FDMXOutputPortConfig& OutputPortConfig)
+void FDMXOutputPort::UpdateFromConfig(FDMXOutputPortConfig& InOutOutputPortConfig, bool bForceUpdateRegistrationWithProtocol)
 {	
 	// Need a valid config for the port
-	OutputPortConfig.MakeValid();
+	InOutOutputPortConfig.MakeValid();
+
+	// Avoid further changes to the config
+	const FDMXOutputPortConfig& OutputPortConfig = InOutOutputPortConfig;
 
 	// Can only use configs that correspond to project settings
 	const UDMXProtocolSettings* ProtocolSettings = GetDefault<UDMXProtocolSettings>();
@@ -439,8 +443,13 @@ void FDMXOutputPort::UpdateFromConfig(FDMXOutputPortConfig& OutputPortConfig)
 	checkf(bConfigIsInProjectSettings, TEXT("Can only use configs with a guid that corresponds to a config in project settings"));
 
 	// Find if the port needs update its registration with the protocol
-	const bool bNeedsUpdateRegistration = [this, &OutputPortConfig]()
+	const bool bNeedsUpdateRegistration = [this, &OutputPortConfig, bForceUpdateRegistrationWithProtocol]()
 	{
+		if (bForceUpdateRegistrationWithProtocol)
+		{
+			return true;
+		}
+
 		if (IsRegistered() != CommunicationDeterminator.IsSendDMXEnabled())
 		{
 			return true;
diff --git a/Engine/Plugins/VirtualProduction/DMX/DMXProtocol/Source/DMXProtocol/Private/IO/DMXPortManager.cpp b/Engine/Plugins/VirtualProduction/DMX/DMXProtocol/Source/DMXProtocol/Private/IO/DMXPortManager.cpp
--- a/Engine/Plugins/VirtualProduction/DMX/DMXProtocol/Source/DMXProtocol/Private/IO/DMXPortManager.cpp
+++ b/Engine/Plugins/VirtualProduction/DMX/DMXProtocol/Source/DMXProtocol/Private/IO/DMXPortManager.cpp
@@ -163,7 +163,7 @@ FDMXOutputPortSharedRef FDMXPortManager::FindOutputPortByGuidChecked(const FGuid
 	return InvalidPtr.ToSharedRef();
 }
 
-void FDMXPortManager::UpdateFromProtocolSettings()
+void FDMXPortManager::UpdateFromProtocolSettings(bool bForceUpdateRegistrationWithProtocol)
 {
 	UDMXProtocolSettings* ProtocolSettings = GetMutableDefault<UDMXProtocolSettings>();
 
@@ -213,7 +213,7 @@ void FDMXPortManager::UpdateFromProtocolSettings()
 		{
 			FDMXInputPortSharedRef InputPort = GetOrCreateInputPortFromConfig(InputPortConfig);
 
-			InputPort->UpdateFromConfig(InputPortConfig);
+			InputPort->UpdateFromConfig(InputPortConfig, bForceUpdateRegistrationWithProtocol);
 
 			PortGuidsFromProtocolSettings.AddUnique(InputPort->GetPortGuid());
 		}
@@ -229,7 +229,7 @@ void FDMXPortManager::UpdateFromProtocolSettings()
 		{
 			FDMXOutputPortSharedRef OutputPort = GetOrCreateOutputPortFromConfig(OutputPortConfig);
 
-			OutputPort->UpdateFromConfig(OutputPortConfig);
+			OutputPort->UpdateFromConfig(OutputPortConfig, bForceUpdateRegistrationWithProtocol);
 
 			PortGuidsFromProtocolSettings.AddUnique(OutputPort->GetPortGuid());
 		}
diff --git a/Engine/Plugins/VirtualProduction/DMX/DMXProtocol/Source/DMXProtocol/Public/IO/DMXInputPort.h b/Engine/Plugins/VirtualProduction/DMX/DMXProtocol/Source/DMXProtocol/Public/IO/DMXInputPort.h
--- a/Engine/Plugins/VirtualProduction/DMX/DMXProtocol/Source/DMXProtocol/Public/IO/DMXInputPort.h
+++ b/Engine/Plugins/VirtualProduction/DMX/DMXProtocol/Source/DMXProtocol/Public/IO/DMXInputPort.h
@@ -44,8 +44,13 @@ public:
 	/** Creates a dmx input port config that corresponds to the port */
 	FDMXInputPortConfig MakeInputPortConfig() const;
 
-	/** Updates the Port to use the config of the InputPortConfig. Makes the config valid if it's invalid. */
-	void UpdateFromConfig(FDMXInputPortConfig& InputPortConfig);
+	/** 
+	 * Updates the Port to use the config of the InputPortConfig. Makes the config valid if it's invalid. 
+	 *
+	 * @param InOutInputPortConfig					The config that is applied. May be changed to a valid config.
+	 * @param bForceUpdateRegistrationWithProtocol	Optional: Forces the port to update its registration with the protocol (useful for runtime changes)
+	 */
+	void UpdateFromConfig(FDMXInputPortConfig& InOutInputPortConfig, bool bForceUpdateRegistrationWithProtocol = false);
 
 	//~ Begin DMXPort Interface 
 	virtual bool IsRegistered() const override;
diff --git a/Engine/Plugins/VirtualProduction/DMX/DMXProtocol/Source/DMXProtocol/Public/IO/DMXInputPortConfig.h b/Engine/Plugins/VirtualProduction/DMX/DMXProtocol/Source/DMXProtocol/Public/IO/DMXInputPortConfig.h
--- a/Engine/Plugins/VirtualProduction/DMX/DMXProtocol/Source/DMXProtocol/Public/IO/DMXInputPortConfig.h
+++ b/Engine/Plugins/VirtualProduction/DMX/DMXProtocol/Source/DMXProtocol/Public/IO/DMXInputPortConfig.h
@@ -89,9 +89,13 @@ public:
 	FORCEINLINE const int32 GetPriority() const { return Priority; }
 
 #if WITH_EDITOR
+	static FName GetPortNamePropertyNameChecked() { return GET_MEMBER_NAME_CHECKED(FDMXInputPortConfig, PortName); }
 	static FName GetProtocolNamePropertyNameChecked() { return GET_MEMBER_NAME_CHECKED(FDMXInputPortConfig, ProtocolName); }
 	static FName GetCommunicationTypePropertyNameChecked() { return GET_MEMBER_NAME_CHECKED(FDMXInputPortConfig, CommunicationType); }
 	static FName GetDeviceAddressPropertyNameChecked() { return GET_MEMBER_NAME_CHECKED(FDMXInputPortConfig, DeviceAddress); }
+	static FName GetLocalUniverseStartPropertyNameChecked() { return GET_MEMBER_NAME_CHECKED(FDMXInputPortConfig, LocalUniverseStart); }
+	static FName GetNumUniversesPropertyNameChecked() { return GET_MEMBER_NAME_CHECKED(FDMXInputPortConfig, NumUniverses); }
+	static FName GetExternUniverseStartPropertyNameChecked() { return GET_MEMBER_NAME_CHECKED(FDMXInputPortConfig, ExternUniverseStart); }
 	static FName GetPriorityStrategyPropertyNameChecked() { return GET_MEMBER_NAME_CHECKED(FDMXInputPortConfig, PriorityStrategy); }
 	static FName GetPriorityPropertyNameChecked() { return GET_MEMBER_NAME_CHECKED(FDMXInputPortConfig, Priority); }
 	static FName GetPortGuidPropertyNameChecked() { return GET_MEMBER_NAME_CHECKED(FDMXInputPortConfig, PortGuid); }
diff --git a/Engine/Plugins/VirtualProduction/DMX/DMXProtocol/Source/DMXProtocol/Public/IO/DMXOutputPort.h b/Engine/Plugins/VirtualProduction/DMX/DMXProtocol/Source/DMXProtocol/Public/IO/DMXOutputPort.h
--- a/Engine/Plugins/VirtualProduction/DMX/DMXProtocol/Source/DMXProtocol/Public/IO/DMXOutputPort.h
+++ b/Engine/Plugins/VirtualProduction/DMX/DMXProtocol/Source/DMXProtocol/Public/IO/DMXOutputPort.h
@@ -116,8 +116,13 @@ public:
 	/** Creates a dmx output port config that corresponds to the port */
 	FDMXOutputPortConfig MakeOutputPortConfig() const;
 
-	/** Updates the Port to use the config of the OutputPortConfig */
-	void UpdateFromConfig(FDMXOutputPortConfig& OutputPortConfig);
+	/**
+	 * Updates the Port to use the config of the OutputPortConfig. Makes the config valid if it's invalid.
+	 *
+	 * @param InOutInputPortConfig					The config that is applied. May be changed to a valid config.
+	 * @param bForceUpdateRegistrationWithProtocol	Optional: Forces the port to update its registration with the protocol (useful for runtime changes)
+	 */
+	void UpdateFromConfig(FDMXOutputPortConfig& OutputPortConfig, bool bForceUpdateRegistrationWithProtocol = false);
 
 	/** Sends DMX over the port */
 	void SendDMX(int32 LocalUniverseID, const TMap<int32, uint8>& ChannelToValueMap);
diff --git a/Engine/Plugins/VirtualProduction/DMX/DMXProtocol/Source/DMXProtocol/Public/IO/DMXOutputPortConfig.h b/Engine/Plugins/VirtualProduction/DMX/DMXProtocol/Source/DMXProtocol/Public/IO/DMXOutputPortConfig.h
--- a/Engine/Plugins/VirtualProduction/DMX/DMXProtocol/Source/DMXProtocol/Public/IO/DMXOutputPortConfig.h
+++ b/Engine/Plugins/VirtualProduction/DMX/DMXProtocol/Source/DMXProtocol/Public/IO/DMXOutputPortConfig.h
@@ -101,10 +101,14 @@ public:
 	FORCEINLINE const FString& GetDestinationAddress() const { return DestinationAddress; }
 
 #if WITH_EDITOR
+	static FName GetPortNamePropertyNameChecked() { return GET_MEMBER_NAME_CHECKED(FDMXOutputPortConfig, PortName); }
 	static FName GetProtocolNamePropertyNameChecked() { return GET_MEMBER_NAME_CHECKED(FDMXOutputPortConfig, ProtocolName); }
 	static FName GetCommunicationTypePropertyNameChecked() { return GET_MEMBER_NAME_CHECKED(FDMXOutputPortConfig, CommunicationType); }
 	static FName GetDeviceAddressPropertyNameChecked() { return GET_MEMBER_NAME_CHECKED(FDMXOutputPortConfig, DeviceAddress); }
 	static FName GetDestinationAddressesPropertyNameChecked() { return GET_MEMBER_NAME_CHECKED(FDMXOutputPortConfig, DestinationAddresses); }
+	static FName GetLocalUniverseStartPropertyNameChecked() { return GET_MEMBER_NAME_CHECKED(FDMXOutputPortConfig, LocalUniverseStart); }
+	static FName GetNumUniversesPropertyNameChecked() { return GET_MEMBER_NAME_CHECKED(FDMXOutputPortConfig, NumUniverses); }
+	static FName GetExternUniverseStartPropertyNameChecked() { return GET_MEMBER_NAME_CHECKED(FDMXOutputPortConfig, ExternUniverseStart); }
 	static FName GetPriorityPropertyNameChecked() { return GET_MEMBER_NAME_CHECKED(FDMXOutputPortConfig, Priority); }
 	static FName GetDelayPropertyNameChecked() { return GET_MEMBER_NAME_CHECKED(FDMXOutputPortConfig, Delay); }
 	static FName GetDelayFrameRatePropertyNameChecked() { return GET_MEMBER_NAME_CHECKED(FDMXOutputPortConfig, DelayFrameRate); }
diff --git a/Engine/Plugins/VirtualProduction/DMX/DMXProtocol/Source/DMXProtocol/Public/IO/DMXPortManager.h b/Engine/Plugins/VirtualProduction/DMX/DMXProtocol/Source/DMXProtocol/Public/IO/DMXPortManager.h
--- a/Engine/Plugins/VirtualProduction/DMX/DMXProtocol/Source/DMXProtocol/Public/IO/DMXPortManager.h
+++ b/Engine/Plugins/VirtualProduction/DMX/DMXProtocol/Source/DMXProtocol/Public/IO/DMXPortManager.h
@@ -98,8 +98,11 @@ public:
 	/** Returns the output port matching the port guid, checked version. */
 	FDMXOutputPortSharedRef FindOutputPortByGuidChecked(const FGuid& PortGuid) const;
 
-	/** Updates ports from protocol settings. Useful when project settings were modified in some way */
-	void UpdateFromProtocolSettings();
+	/** 
+	 * Updates ports from protocol settings. Useful when project settings were modified in some way. 
+	 * If optional bForceUpdateRegistrationWithProtocol is set to true, ports will reregister with the protocol.
+	*/
+	void UpdateFromProtocolSettings(bool bForceUpdateRegistrationWithProtocol = false);
 
 	/** 
 	 * Thread-safe: Suspends the protocols irrevocably, and with it all sending and receiving of DMX via the protocols. 
