DMX - Fix Port Reference variables in BPs cannot be compiled when selected

diff --git a/Engine/Plugins/VirtualProduction/DMX/DMXProtocol/Source/DMXProtocolEditor/Private/DetailsCustomizations/DMXPortReferenceCustomizationBase.cpp b/Engine/Plugins/VirtualProduction/DMX/DMXProtocol/Source/DMXProtocolEditor/Private/DetailsCustomizations/DMXPortReferenceCustomizationBase.cpp
--- a/Engine/Plugins/VirtualProduction/DMX/DMXProtocol/Source/DMXProtocolEditor/Private/DetailsCustomizations/DMXPortReferenceCustomizationBase.cpp
+++ b/Engine/Plugins/VirtualProduction/DMX/DMXProtocol/Source/DMXProtocolEditor/Private/DetailsCustomizations/DMXPortReferenceCustomizationBase.cpp
@@ -79,39 +79,36 @@ void FDMXPortReferenceCustomizationBase::CustomizeChildren(TSharedRef<IPropertyH
 	// a) This is an existing port reference, with a corresponding port config
 	// b) This is an existing port reference, but the corresponding port config got deleted in settings
 	// c) This is a new port reference and doesn't point to any port yet
-	// d) There are no ports specified in project settings
+	// d) The port references are multiselected
 
-	TSharedPtr<FDMXPort, ESPMode::ThreadSafe> Port = FindPortItem();
-	if(Port.IsValid())
+	const TArray<FDMXPortSharedPtr> Ports = FindPortItems();
+	if(Ports.Num() == 1 && Ports[0].IsValid())
 	{
-		// a) This is an existing port reference, with a corresponding port config
-		PortSelector->SelectPort(Port->GetPortGuid());
+		// a) This is a single existing port reference, with a corresponding port config
+		const FGuid PortGuid = Ports[0]->GetPortGuid();
+		PortSelector->SelectPort(PortGuid);
 	}
 	else
 	{
-		FGuid PortGuid = GetPortGuid();
-		if (PortGuid.IsValid())
+		const TArray<FGuid> PortGuids = GetPortGuids();
+		if (PortGuids.Num() == 1 && PortGuids[0].IsValid())
 		{
 			// b) This is an existing port reference, but the corresponding port config got deleted in settings
 
 			// Let the port remain invalid, but let users know
-			ErrorText = LOCTEXT("PortReferenceNoLongerValid", "The referenced Port was deleted. Please select another port.");
-			UE_LOG(LogDMXProtocol, Error, TEXT("The referenced Port was deleted. Please review your libraries."));
+			UE_LOG(LogDMXProtocol, Error, TEXT("The referenced Port was deleted from project settings."));
+		}
+		else if(PortGuids.Num() == 0 && PortSelector->HasSelection())
+		{
+			// c) This is a new port reference and doesn't point to any port yet
+			//	  The port selector makes an initial selection already, so if it's valid there are ports
+			//    It only needs be applied on the new reference.
+			ApplySelectedPortGuid();
 		}
 		else
 		{
-			if (PortSelector->HasSelection())
-			{
-				// c) This is a new port reference and doesn't point to any port yet
-				//	  The port selector makes an initial selection already, so if it's valid there are ports
-				//    It only needs be applied on the new reference.
-				ApplySelectedPortGuid();
-			}
-			else
-			{
-				// d) There are no ports specified in project settings
-				ErrorText = LOCTEXT("NoPortAvailableText", "No ports available. Please create DMX Ports in Project Settings -> Plugins -> DMX Plugin.");
-			}
+			// d) This is multi selected in some way
+			PortSelector->SetHasMultipleValues();
 		}
 	}
 
@@ -133,32 +130,36 @@ void FDMXPortReferenceCustomizationBase::OnPortsChanged()
 	}
 }
 
-TSharedPtr<FDMXPort, ESPMode::ThreadSafe> FDMXPortReferenceCustomizationBase::FindPortItem() const
+TArray<FDMXPortSharedPtr> FDMXPortReferenceCustomizationBase::FindPortItems() const
 {
-	const FGuid PortGuid = GetPortGuid();
+	const TArray<FGuid> PortGuids = GetPortGuids();
 
-	if (PortGuid.IsValid())
+	TArray<FDMXPortSharedPtr> Result;
+	for (const FGuid& PortGuid : PortGuids)
 	{
-		const FDMXInputPortSharedRef* InputPortPtr = FDMXPortManager::Get().GetInputPorts().FindByPredicate([&PortGuid](const FDMXInputPortSharedRef& InputPort) {
-			return InputPort->GetPortGuid() == PortGuid;
-			});
-
-		if (InputPortPtr)
+		if (PortGuid.IsValid())
 		{
-			return *InputPortPtr;
-		}
+			const FDMXInputPortSharedRef* InputPortPtr = FDMXPortManager::Get().GetInputPorts().FindByPredicate([&PortGuid](const FDMXInputPortSharedRef& InputPort) {
+				return InputPort->GetPortGuid() == PortGuid;
+				});
+
+			if (InputPortPtr)
+			{
+				Result.Add(*InputPortPtr);
+			}
 
-		const FDMXOutputPortSharedRef* OutputPortPtr = FDMXPortManager::Get().GetOutputPorts().FindByPredicate([&PortGuid](const FDMXOutputPortSharedRef& OutputPort) {
+			const FDMXOutputPortSharedRef* OutputPortPtr = FDMXPortManager::Get().GetOutputPorts().FindByPredicate([&PortGuid](const FDMXOutputPortSharedRef& OutputPort) {
 				return OutputPort->GetPortGuid() == PortGuid;
-			});
+				});
 
-		if (OutputPortPtr)
-		{
-			return *OutputPortPtr;
+			if (OutputPortPtr)
+			{
+				Result.Add(*OutputPortPtr);
+			}
 		}
 	}
 
-	return nullptr;
+	return Result;
 }
 
 void FDMXPortReferenceCustomizationBase::ApplySelectedPortGuid()
@@ -181,43 +182,44 @@ void FDMXPortReferenceCustomizationBase::ApplySelectedPortGuid()
 		return FGuid();
 	}();
 
-	const TSharedPtr<IPropertyHandle>& PortGuidHandle = GetPortGuidHandle();
-	check(PortGuidHandle.IsValid() && PortGuidHandle->IsValidHandle());
+	if (!GetPortGuids().Contains(SelectedGuid))
+	{
+		const TSharedPtr<IPropertyHandle>& PortGuidHandle = GetPortGuidHandle();
+		check(PortGuidHandle.IsValid() && PortGuidHandle->IsValidHandle());
 
-	PortGuidHandle->NotifyPreChange();
+		PortGuidHandle->NotifyPreChange();
 
-	TArray<void*> RawData;
-	PortGuidHandle->AccessRawData(RawData);
+		TArray<void*> RawDataArray;
+		PortGuidHandle->AccessRawData(RawDataArray);
 
-	for(void* RawElement : RawData)
-	{
-		FMemory::Memcpy(RawElement, &SelectedGuid, sizeof(FGuid));
-	}
+		for (void* RawData : RawDataArray)
+		{
+			FMemory::Memcpy(RawData, &SelectedGuid, sizeof(FGuid));
+		}
 
-	PortGuidHandle->NotifyPostChange(EPropertyChangeType::ValueSet);
+		PortGuidHandle->NotifyPostChange(EPropertyChangeType::ValueSet);
+	}
 }
 
-FGuid FDMXPortReferenceCustomizationBase::GetPortGuid() const
+TArray<FGuid> FDMXPortReferenceCustomizationBase::GetPortGuids() const
 {
 	const TSharedPtr<IPropertyHandle>& PortGuidHandle = GetPortGuidHandle();
 	check(PortGuidHandle.IsValid() && PortGuidHandle->IsValidHandle());
 
-	TArray<void*> RawData;
-	PortGuidHandle->AccessRawData(RawData);
+	TArray<void*> RawDataArray;
+	PortGuidHandle->AccessRawData(RawDataArray);
 
-	if (ensureMsgf(RawData.Num() == 1, TEXT("The port configs reside in an array, multi-editing is unexpected")))
+	TArray<FGuid> Result;
+	for (const void* RawData : RawDataArray)
 	{
-		FGuid* PortGuidPtr = reinterpret_cast<FGuid*>(RawData[0]);
-		if (ensureMsgf(PortGuidPtr, TEXT("Expected to be valid")))
+		const FGuid* PortGuidPtr = reinterpret_cast<const FGuid*>(RawData);
+		if (PortGuidPtr && PortGuidPtr->IsValid())
 		{
-			if (PortGuidPtr->IsValid())
-			{
-				return *PortGuidPtr;
-			};
+			Result.Add(*PortGuidPtr);
 		}
 	}
 
-	return FGuid();
+	return Result;
 }
 
 #undef LOCTEXT_NAMESPACE
diff --git a/Engine/Plugins/VirtualProduction/DMX/DMXProtocol/Source/DMXProtocolEditor/Private/DetailsCustomizations/DMXPortReferenceCustomizationBase.h b/Engine/Plugins/VirtualProduction/DMX/DMXProtocol/Source/DMXProtocolEditor/Private/DetailsCustomizations/DMXPortReferenceCustomizationBase.h
--- a/Engine/Plugins/VirtualProduction/DMX/DMXProtocol/Source/DMXProtocolEditor/Private/DetailsCustomizations/DMXPortReferenceCustomizationBase.h
+++ b/Engine/Plugins/VirtualProduction/DMX/DMXProtocol/Source/DMXProtocolEditor/Private/DetailsCustomizations/DMXPortReferenceCustomizationBase.h
@@ -2,6 +2,8 @@
 
 #pragma once
 
+#include "DMXProtocolCommon.h"
+
 #include "CoreMinimal.h"
 #include "IPropertyTypeCustomization.h"
 #include "Templates/SharedPointer.h"
@@ -43,13 +45,13 @@ private:
 	void OnPortsChanged();
 
 	/** Helper function that finds the corresponding input port, or nullptr if no corresponding port */
-	TSharedPtr<FDMXPort, ESPMode::ThreadSafe> FindPortItem() const;
+	TArray<FDMXPortSharedPtr> FindPortItems() const;
 
 	/** Applies the selected port Guid to the customized struct */
 	void ApplySelectedPortGuid();
 
-	/** Helper function that gets the Guid of the edited port */
-	FGuid GetPortGuid() const;
+	/** Helper function that gets the Guid of the edited port(s) */
+	TArray<FGuid> GetPortGuids() const;
 
 	/** Port selector widget */
 	TSharedPtr<SDMXPortSelector> PortSelector;
diff --git a/Engine/Plugins/VirtualProduction/DMX/DMXProtocol/Source/DMXProtocolEditor/Private/Widgets/SDMXPortSelector.cpp b/Engine/Plugins/VirtualProduction/DMX/DMXProtocol/Source/DMXProtocolEditor/Private/Widgets/SDMXPortSelector.cpp
--- a/Engine/Plugins/VirtualProduction/DMX/DMXProtocol/Source/DMXProtocolEditor/Private/Widgets/SDMXPortSelector.cpp
+++ b/Engine/Plugins/VirtualProduction/DMX/DMXProtocol/Source/DMXProtocolEditor/Private/Widgets/SDMXPortSelector.cpp
@@ -77,39 +77,50 @@ void SDMXPortSelector::Construct(const FArguments& InArgs)
 
 bool SDMXPortSelector::IsInputPortSelected() const
 {
-	const TSharedPtr<FDMXPortSelectorItem> CurrentSelection = PortNameComboBox->GetSelectedItem();
-	if (CurrentSelection.IsValid())
+	if (!bHasMultipleValues)
 	{
-		return CurrentSelection->GetType() == EDMXPortSelectorItemType::Input;
+		const TSharedPtr<FDMXPortSelectorItem> CurrentSelection = PortNameComboBox->GetSelectedItem();
+		if (CurrentSelection.IsValid())
+		{
+			return CurrentSelection->GetType() == EDMXPortSelectorItemType::Input;
+		}
 	}
+
 	return false;
 }
 
 bool SDMXPortSelector::IsOutputPortSelected() const
 {
-	const TSharedPtr<FDMXPortSelectorItem> CurrentSelection = PortNameComboBox->GetSelectedItem();
-	if (CurrentSelection.IsValid())
+	if (!bHasMultipleValues)
 	{
-		return CurrentSelection->GetType() == EDMXPortSelectorItemType::Output;
+		const TSharedPtr<FDMXPortSelectorItem> CurrentSelection = PortNameComboBox->GetSelectedItem();
+		if (CurrentSelection.IsValid())
+		{
+			return CurrentSelection->GetType() == EDMXPortSelectorItemType::Output;
+		}
 	}
+
 	return false;
 }
 
 FDMXInputPortSharedPtr SDMXPortSelector::GetSelectedInputPort() const
 {
-	const TSharedPtr<FDMXPortSelectorItem> CurrentSelection = PortNameComboBox->GetSelectedItem();
-	if (CurrentSelection.IsValid() &&
-		CurrentSelection->GetType() == EDMXPortSelectorItemType::Input)
+	if (!bHasMultipleValues)
 	{
-		const TArray<FDMXInputPortSharedRef>& InputPorts = FDMXPortManager::Get().GetInputPorts();
+		const TSharedPtr<FDMXPortSelectorItem> CurrentSelection = PortNameComboBox->GetSelectedItem();
+		if (CurrentSelection.IsValid() &&
+			CurrentSelection->GetType() == EDMXPortSelectorItemType::Input)
+		{
+			const TArray<FDMXInputPortSharedRef>& InputPorts = FDMXPortManager::Get().GetInputPorts();
 
-		const FDMXInputPortSharedRef* SelectedPortPtr = InputPorts.FindByPredicate([&CurrentSelection](const FDMXInputPortSharedRef& InputPort) {
-			return InputPort->GetPortGuid() == CurrentSelection->GetGuid();
-		});
+			const FDMXInputPortSharedRef* SelectedPortPtr = InputPorts.FindByPredicate([&CurrentSelection](const FDMXInputPortSharedRef& InputPort) {
+				return InputPort->GetPortGuid() == CurrentSelection->GetGuid();
+				});
 
-		if (SelectedPortPtr)
-		{
-			return *SelectedPortPtr;
+			if (SelectedPortPtr)
+			{
+				return *SelectedPortPtr;
+			}
 		}
 	}
 
@@ -118,19 +129,22 @@ FDMXInputPortSharedPtr SDMXPortSelector::GetSelectedInputPort() const
 
 FDMXOutputPortSharedPtr SDMXPortSelector::GetSelectedOutputPort() const
 {
-	const TSharedPtr<FDMXPortSelectorItem> CurrentSelection = PortNameComboBox->GetSelectedItem();
-	if (CurrentSelection.IsValid() &&
-		CurrentSelection->GetType() == EDMXPortSelectorItemType::Output)
+	if (!bHasMultipleValues)
 	{
-		const TArray<FDMXOutputPortSharedRef>& OutputPorts = FDMXPortManager::Get().GetOutputPorts();
+		const TSharedPtr<FDMXPortSelectorItem> CurrentSelection = PortNameComboBox->GetSelectedItem();
+		if (CurrentSelection.IsValid() &&
+			CurrentSelection->GetType() == EDMXPortSelectorItemType::Output)
+		{
+			const TArray<FDMXOutputPortSharedRef>& OutputPorts = FDMXPortManager::Get().GetOutputPorts();
 
-		const FDMXOutputPortSharedRef* SelectedPortPtr = OutputPorts.FindByPredicate([&CurrentSelection](const FDMXOutputPortSharedRef& OutputPort) {
-			return OutputPort->GetPortGuid() == CurrentSelection->GetGuid();
-			});
+			const FDMXOutputPortSharedRef* SelectedPortPtr = OutputPorts.FindByPredicate([&CurrentSelection](const FDMXOutputPortSharedRef& OutputPort) {
+				return OutputPort->GetPortGuid() == CurrentSelection->GetGuid();
+				});
 
-		if (SelectedPortPtr)
-		{
-			return *SelectedPortPtr;
+			if (SelectedPortPtr)
+			{
+				return *SelectedPortPtr;
+			}
 		}
 	}
 
@@ -144,6 +158,8 @@ bool SDMXPortSelector::HasSelection() const
 
 void SDMXPortSelector::SelectPort(const FGuid& PortGuid)
 {
+	bHasMultipleValues = false;
+
 	TSharedPtr<FDMXPortSelectorItem> PortItem = FindPortItem(PortGuid);
 	if (PortItem != PortNameComboBox->GetSelectedItem())
 	{
@@ -160,6 +176,13 @@ void SDMXPortSelector::SelectPort(const FGuid& PortGuid)
 	}
 }
 
+void SDMXPortSelector::SetHasMultipleValues()
+{
+	bHasMultipleValues = true;
+
+	PortNameTextBlock->SetText(LOCTEXT("ShowAsMultipleValues", "Mutiple Values"));
+}
+
 TSharedRef<SWidget> SDMXPortSelector::GenerateComboBoxEntry(TSharedPtr<FDMXPortSelectorItem> Item)
 {
 	if (Item->IsTitleRow())
@@ -188,6 +211,8 @@ void SDMXPortSelector::OnPortItemSelectionChanged(TSharedPtr<FDMXPortSelectorIte
 	
 	if (Item.IsValid())
 	{
+		bHasMultipleValues = false;
+
 		if (Item->IsTitleRow())
 		{
 			// Restore the previous selection if possible
diff --git a/Engine/Plugins/VirtualProduction/DMX/DMXProtocol/Source/DMXProtocolEditor/Public/Widgets/SDMXPortSelector.h b/Engine/Plugins/VirtualProduction/DMX/DMXProtocol/Source/DMXProtocolEditor/Public/Widgets/SDMXPortSelector.h
--- a/Engine/Plugins/VirtualProduction/DMX/DMXProtocol/Source/DMXProtocolEditor/Public/Widgets/SDMXPortSelector.h
+++ b/Engine/Plugins/VirtualProduction/DMX/DMXProtocol/Source/DMXProtocolEditor/Public/Widgets/SDMXPortSelector.h
@@ -108,6 +108,12 @@ public:
 	/** Selects specified Port by Port Guid. If the Guid is not valid, clears the selection */
 	void SelectPort(const FGuid& PortGuid);
 
+	/** 
+	 * Instead of showing a selected port, the widget displays 'Multiple Values' and no longer returns a selected port.
+	 * The widget will no longer display 'Multiple Values' when a selection was made. 
+	 */
+	void SetHasMultipleValues();
+
 private:
 	/** Called to generate entries in the port name combo box */
 	TSharedRef<SWidget> GenerateComboBoxEntry(TSharedPtr<FDMXPortSelectorItem> Item);
@@ -130,6 +136,9 @@ private:
 	/** Called when ports changed */
 	void OnPortsChanged();
 
+	/** State where the widget was set to have multiple values */
+	bool bHasMultipleValues = false;
+
 	/** The combo box that holds all ports */
 	TSharedPtr<SComboBox<TSharedPtr<FDMXPortSelectorItem>>> PortNameComboBox;
 
