DMX - Fix issues when setting delay, delay frame rate and sACN universes via command line

diff --git a/Engine/Plugins/VirtualProduction/DMX/DMXProtocol/Source/DMXProtocol/Private/IO/DMXOutputPort.cpp b/Engine/Plugins/VirtualProduction/DMX/DMXProtocol/Source/DMXProtocol/Private/IO/DMXOutputPort.cpp
--- a/Engine/Plugins/VirtualProduction/DMX/DMXProtocol/Source/DMXProtocol/Private/IO/DMXOutputPort.cpp
+++ b/Engine/Plugins/VirtualProduction/DMX/DMXProtocol/Source/DMXProtocol/Private/IO/DMXOutputPort.cpp
@@ -12,6 +12,7 @@
 
 #include "HAL/IConsoleManager.h"
 #include "HAL/RunnableThread.h"
+#include "Misc/FrameRate.h"
 
 
 #define LOCTEXT_NAMESPACE "DMXOutputPort"
@@ -251,7 +252,7 @@ static FAutoConsoleCommand GDMXSetOutputPortPriorityCommand(
 
 static FAutoConsoleCommand GDMXSetOutputPortDelayCommand(
 	TEXT("DMX.SetOutputPortDelay"),
-	TEXT("DMX.SetOutputPortDelay [PortName][Delay]. Sets the delay of the output port (depending on the delay frame rate). Example: DMX.SetOutputPortDelay MyOutputPort 3.54"),
+	TEXT("DMX.SetOutputPortDelay [PortName][Delay][(optional)FrameRate]. Sets the delay of the output port, optionally with a frame rate (e.g. '30fps', '0.001s' or '12000/1001'). Example: DMX.SetOutputPortDelay MyOutputPort 10.5 30fps"),
 	FConsoleCommandWithArgsDelegate::CreateStatic(
 		[](const TArray<FString>& Args)
 		{
@@ -265,42 +266,21 @@ static FAutoConsoleCommand GDMXSetOutputPortDelayCommand(
 			double DelayValue;
 			if (LexTryParseString<double>(DelayValue, *DelayValueString))
 			{
-				if (DelayValue >= 0.0)
+				FFrameRate NewFrameRate = FFrameRate(1, 1);
+				
+				// Parse optional frame rate argument
+				if (Args.Num() == 3)
 				{
-					const FDMXOutputPortSharedRef* PortPtr = FDMXPortManager::Get().GetOutputPorts().FindByPredicate([PortName](const FDMXOutputPortSharedPtr& OutputPort)
-						{
-							return OutputPort->GetPortName() == PortName;
-						});
-
-					if (PortPtr)
+					const FString& DelayFrameRateString = Args[2];
+				
+					TValueOrError<FFrameRate, FExpressionError> ParseResult = ParseFrameRate(*DelayFrameRateString);
+					if (ParseResult.IsValid())
 					{
-						const FDMXOutputPortConfig PortConfig = (*PortPtr)->MakeOutputPortConfig();
-						const double DelaySeconds = FMath::Min(DelayValue * PortConfig.GetDelayFrameRate().AsDecimal(), 60.f);
-
-						DMX_OVERRIDE_OUTPUTPORT_VAR(DelaySeconds, PortName, DelayValue);
+						NewFrameRate = ParseResult.GetValue();
 					}
 				}
-			}
-		})
-);
 
-static FAutoConsoleCommand GDMXSetOutputPortDelayFrameRateCommand(
-	TEXT("DMX.SetOutputPortDelayFrameRate"),
-	TEXT("DMX.SetOutputPortDelayFrameRate [PortName][FrameRate]. Sets the frame rate of the delay of the output port (1.0 for seconds). Example: DMX.SetOutputPortDelayFrameRate MyOutputPort 33.3"),
-	FConsoleCommandWithArgsDelegate::CreateStatic(
-		[](const TArray<FString>& Args)
-		{
-			if (Args.Num() < 2)
-			{
-				return;
-			}
-
-			const FString& PortName = Args[0];
-			const FString& DelayFrameRateValueString = Args[1];
-			double DelayFrameRateValue;
-			if (LexTryParseString<double>(DelayFrameRateValue, *DelayFrameRateValueString))
-			{
-				if (DelayFrameRateValue > 0.0)
+				if (DelayValue >= 0.0)
 				{
 					const FDMXOutputPortSharedRef* PortPtr = FDMXPortManager::Get().GetOutputPorts().FindByPredicate([PortName](const FDMXOutputPortSharedPtr& OutputPort)
 						{
@@ -309,13 +289,11 @@ static FAutoConsoleCommand GDMXSetOutputPortDelayFrameRateCommand(
 
 					if (PortPtr)
 					{
-						const FDMXOutputPortConfig PortConfig = (*PortPtr)->MakeOutputPortConfig();
-						const double OldDelay = PortConfig.GetDelaySeconds() / PortConfig.GetDelayFrameRate().AsDecimal();
-						
-						const FFrameRate NewDelayFrameRate = FFrameRate(1, DelayFrameRateValue);
-						const double NewDelaySeconds = FMath::Min(OldDelay * NewDelayFrameRate.AsDecimal(), 60.f);
+						// Limit to 60 seconds
+						const double FrameRateDecimal = static_cast<double>(NewFrameRate.Numerator) / NewFrameRate.Denominator;
+						const double NewDelaySeconds = FMath::Min(DelayValue / FrameRateDecimal, 60.f * FrameRateDecimal);
 
-						DMX_OVERRIDE_OUTPUTPORT_VAR(DelayFrameRate, PortName, NewDelayFrameRate);
+						DMX_OVERRIDE_OUTPUTPORT_VAR(DelayFrameRate, PortName, NewFrameRate);
 						DMX_OVERRIDE_OUTPUTPORT_VAR(DelaySeconds, PortName, NewDelaySeconds);
 					}
 				}
diff --git a/Engine/Plugins/VirtualProduction/DMX/DMXProtocol/Source/DMXProtocol/Private/IO/DMXOutputPortConfig.cpp b/Engine/Plugins/VirtualProduction/DMX/DMXProtocol/Source/DMXProtocol/Private/IO/DMXOutputPortConfig.cpp
--- a/Engine/Plugins/VirtualProduction/DMX/DMXProtocol/Source/DMXProtocol/Private/IO/DMXOutputPortConfig.cpp
+++ b/Engine/Plugins/VirtualProduction/DMX/DMXProtocol/Source/DMXProtocol/Private/IO/DMXOutputPortConfig.cpp
@@ -55,7 +55,7 @@ FDMXOutputPortConfig::FDMXOutputPortConfig(const FGuid& InPortGuid, const FDMXOu
 	, NumUniverses(InitializationData.NumUniverses)
 	, ExternUniverseStart(InitializationData.ExternUniverseStart)
 	, Priority(InitializationData.Priority)
-	, Delay(InitializationData.DelaySeconds)
+	, Delay(InitializationData.DelaySeconds * InitializationData.DelayFrameRate.AsDecimal())
 	, DelayFrameRate(InitializationData.DelayFrameRate)
 	, PortGuid(InPortGuid)
 {
diff --git a/Engine/Plugins/VirtualProduction/DMX/DMXProtocol/Source/DMXProtocolSACN/Private/DMXProtocolSACNReceiver.cpp b/Engine/Plugins/VirtualProduction/DMX/DMXProtocol/Source/DMXProtocolSACN/Private/DMXProtocolSACNReceiver.cpp
--- a/Engine/Plugins/VirtualProduction/DMX/DMXProtocol/Source/DMXProtocolSACN/Private/DMXProtocolSACNReceiver.cpp
+++ b/Engine/Plugins/VirtualProduction/DMX/DMXProtocol/Source/DMXProtocolSACN/Private/DMXProtocolSACNReceiver.cpp
@@ -142,32 +142,26 @@ void FDMXProtocolSACNReceiver::UnassignInputPort(const TSharedPtr<FDMXInputPort,
 	check(AssignedInputPorts.Contains(InputPort));
 	AssignedInputPorts.Remove(InputPort);
 
-	ISocketSubsystem* SocketSubsystem = ISocketSubsystem::Get(PLATFORM_SOCKETSUBSYSTEM);
-	TSharedRef<FInternetAddr> MulticastGroupAddrToLeave = SocketSubsystem->CreateInternetAddr();
-
 	// Leave multicast groups outside of remaining port's universe ranges
-	for (const TTuple<uint32, uint16>& MulticastGroupAddrToUniverseIDKvp : MulticastGroupAddrToUniverseIDMap)
+	for (const TTuple<uint32, uint16>& MulticastGroupAddrToUniverseIDKvp : TMap<uint32, uint16>(MulticastGroupAddrToUniverseIDMap))
 	{
-		MulticastGroupAddrToLeave->SetIp(MulticastGroupAddrToUniverseIDKvp.Key);
-		MulticastGroupAddrToLeave->SetPort(ACN_PORT);
-		uint16 UniverseID = MulticastGroupAddrToUniverseIDKvp.Value;
-
-		bool bGroupInUse = true;
-		for (const FDMXInputPortSharedPtr& Port : AssignedInputPorts)
-		{
-			if (Port->GetExternUniverseStart() >= UniverseID &&
-				Port->GetExternUniverseEnd() <= UniverseID)
+		const uint32 MulticastGroupAddr = MulticastGroupAddrToUniverseIDKvp.Key;
+		const uint16 UniverseID = MulticastGroupAddrToUniverseIDKvp.Value;
+		const bool bUniverseStillInUse = AssignedInputPorts.Array().ContainsByPredicate([UniverseID](const FDMXInputPortSharedPtr& OtherInputPort)
 			{
-				continue;
-			}
+				return OtherInputPort->IsExternUniverseInPortRange(UniverseID);
+			});
 
-			bGroupInUse = false;
-			break;
-		}
-
-		if (!bGroupInUse)
+		if (!bUniverseStillInUse)
 		{
+			ISocketSubsystem* SocketSubsystem = ISocketSubsystem::Get(PLATFORM_SOCKETSUBSYSTEM);
+			TSharedRef<FInternetAddr> MulticastGroupAddrToLeave = SocketSubsystem->CreateInternetAddr();
+			MulticastGroupAddrToLeave->SetIp(MulticastGroupAddr);
+			MulticastGroupAddrToLeave->SetPort(ACN_PORT);
+
 			Socket->LeaveMulticastGroup(*MulticastGroupAddrToLeave);
+
+			MulticastGroupAddrToUniverseIDMap.Remove(MulticastGroupAddr);
 		}
 	}
 }
