<saved by Perforce>

diff --git a/Engine/Plugins/VirtualProduction/DMX/DMXProtocol/Source/DMXProtocol/Private/DMXProtocolSettings.cpp b/Engine/Plugins/VirtualProduction/DMX/DMXProtocol/Source/DMXProtocol/Private/DMXProtocolSettings.cpp
--- a/Engine/Plugins/VirtualProduction/DMX/DMXProtocol/Source/DMXProtocol/Private/DMXProtocolSettings.cpp
+++ b/Engine/Plugins/VirtualProduction/DMX/DMXProtocol/Source/DMXProtocol/Private/DMXProtocolSettings.cpp
@@ -252,3 +252,51 @@ void UDMXProtocolSettings::OverrideReceiveDMXEnabled(bool bEnabled)
 
 	OnSetReceiveDMXEnabled.Broadcast(bEnabled);
 }
+
+FString UDMXProtocolSettings::GetUniqueInputPortName() const
+{
+	const FString AutoNamePrefix = "InputPort";
+	int32 HighestPortNumber = 0;
+
+	for (const FDMXInputPortConfig InputPortConfig : InputPortConfigs)
+	{
+		const FString& PortName = InputPortConfig.GetPortName();
+		if (PortName.StartsWith(AutoNamePrefix))
+		{
+			int32 PortNameNumericSuffix = 0;
+			if (LexTryParseString<int32>(PortNameNumericSuffix, *PortName.RightChop(AutoNamePrefix.Len())))
+			{
+				if (PortNameNumericSuffix > HighestPortNumber)
+				{
+					HighestPortNumber = PortNameNumericSuffix;
+				}
+			}
+		}
+	}
+
+	return AutoNamePrefix + FString::FromInt(HighestPortNumber + 1);
+}
+
+FString UDMXProtocolSettings::GetUniqueOutputPortName() const
+{
+	const FString AutoNamePrefix = "OutputPort";
+	int32 HighestPortNumber = 0;
+
+	for (const FDMXOutputPortConfig& OutputPortConfig : OutputPortConfigs)
+	{
+		const FString& PortName = OutputPortConfig.GetPortName();
+		if (PortName.StartsWith(AutoNamePrefix))
+		{
+			int32 PortNameNumericSuffix = 0;
+			if (LexTryParseString<int32>(PortNameNumericSuffix, *PortName.RightChop(AutoNamePrefix.Len())))
+			{
+				if (PortNameNumericSuffix > HighestPortNumber)
+				{
+					HighestPortNumber = PortNameNumericSuffix;
+				}
+			}
+		}
+	}
+
+	return AutoNamePrefix + FString::FromInt(HighestPortNumber + 1);
+}
diff --git a/Engine/Plugins/VirtualProduction/DMX/DMXProtocol/Source/DMXProtocol/Private/IO/DMXInputPortConfig.cpp b/Engine/Plugins/VirtualProduction/DMX/DMXProtocol/Source/DMXProtocol/Private/IO/DMXInputPortConfig.cpp
--- a/Engine/Plugins/VirtualProduction/DMX/DMXProtocol/Source/DMXProtocol/Private/IO/DMXInputPortConfig.cpp
+++ b/Engine/Plugins/VirtualProduction/DMX/DMXProtocol/Source/DMXProtocol/Private/IO/DMXInputPortConfig.cpp
@@ -34,8 +34,6 @@ FDMXInputPortConfig::FDMXInputPortConfig(const FGuid& InPortGuid)
 	check(FModuleManager::Get().IsModuleLoaded("DMXProtocol"));
 	check(PortGuid.IsValid());
 
-	GenerateUniquePortName();
-
 	MakeValid();
 }
 
@@ -56,8 +54,6 @@ FDMXInputPortConfig::FDMXInputPortConfig(const FGuid& InPortGuid, const FDMXInpu
 	check(PortGuid.IsValid());
 	check(!ProtocolName.IsNone())
 
-	GenerateUniquePortName();
-
 	MakeValid();
 }
 
@@ -118,6 +114,12 @@ void FDMXInputPortConfig::MakeValid()
 			}
 		}
 	}
+
+	if (PortName.IsEmpty())
+	{
+		UDMXProtocolSettings* ProtocolSettings = GetMutableDefault<UDMXProtocolSettings>();
+		PortName = ProtocolSettings->GetUniqueInputPortName();
+	}
 }
 
 FString FDMXInputPortConfig::GetDeviceAddress() const
@@ -131,30 +133,4 @@ FString FDMXInputPortConfig::GetDeviceAddress() const
 		return OverrideIP;
 	}
 	return DeviceAddress;
-}
-
-void FDMXInputPortConfig::GenerateUniquePortName()
-{
-	if (!PortName.IsEmpty())
-	{
-		return;
-	}
-
-	const UDMXProtocolSettings* ProtocolSettings = GetDefault<UDMXProtocolSettings>();
-	check(ProtocolSettings);
-
-	TSet<FString> OtherPortNames;
-	for (const FDMXInputPortConfig& InputPortConfig : ProtocolSettings->InputPortConfigs)
-	{
-		if (&InputPortConfig == this)
-		{
-			continue;
-		}
-
-		OtherPortNames.Add(InputPortConfig.PortName);
-	}
-
-	FString BaseName = TEXT("InputPort_1");
-
-	PortName = FDMXProtocolUtils::GenerateUniqueNameFromExisting(OtherPortNames, BaseName);
-}
+}
\ No newline at end of file
diff --git a/Engine/Plugins/VirtualProduction/DMX/DMXProtocol/Source/DMXProtocol/Private/IO/DMXOutputPortConfig.cpp b/Engine/Plugins/VirtualProduction/DMX/DMXProtocol/Source/DMXProtocol/Private/IO/DMXOutputPortConfig.cpp
--- a/Engine/Plugins/VirtualProduction/DMX/DMXProtocol/Source/DMXProtocol/Private/IO/DMXOutputPortConfig.cpp
+++ b/Engine/Plugins/VirtualProduction/DMX/DMXProtocol/Source/DMXProtocol/Private/IO/DMXOutputPortConfig.cpp
@@ -39,8 +39,6 @@ FDMXOutputPortConfig::FDMXOutputPortConfig(const FGuid& InPortGuid)
 	check(FModuleManager::Get().IsModuleLoaded("DMXProtocol"));
 	check(PortGuid.IsValid());
 
-	GenerateUniquePortName();
-
 	MakeValid();
 }
 
@@ -64,8 +62,6 @@ FDMXOutputPortConfig::FDMXOutputPortConfig(const FGuid& InPortGuid, const FDMXOu
 	check(PortGuid.IsValid());
 	check(!ProtocolName.IsNone())
 
-	GenerateUniquePortName();
-
 	MakeValid();
 }
 
@@ -135,6 +131,12 @@ void FDMXOutputPortConfig::MakeValid()
 	{
 		Delay = 0.0;
 	}
+
+	if (PortName.IsEmpty())
+	{
+		UDMXProtocolSettings* ProtocolSettings = GetMutableDefault<UDMXProtocolSettings>();
+		PortName = ProtocolSettings->GetUniqueOutputPortName();
+	}
 }
 
 int32 FDMXOutputPortConfig::GetDelaySeconds() const
@@ -153,29 +155,4 @@ FString FDMXOutputPortConfig::GetDeviceAddress() const
 		return OverrideIP;
 	}
 	return DeviceAddress;
-}
-
-void FDMXOutputPortConfig::GenerateUniquePortName()
-{
-	if (!PortName.IsEmpty())
-	{
-		return;
-	}
-
-	const UDMXProtocolSettings* ProtocolSettings = GetDefault<UDMXProtocolSettings>();
-	check(ProtocolSettings);
-
-	TSet<FString> OtherPortNames;
-	for (const FDMXOutputPortConfig& PortConfig : ProtocolSettings->OutputPortConfigs)
-	{
-		if (&PortConfig == this)
-		{
-			continue;
-		}
-
-		OtherPortNames.Add(PortConfig.PortName);
-	}
-
-	FString BaseName = TEXT("OutputPort_1");
-	PortName = FDMXProtocolUtils::GenerateUniqueNameFromExisting(OtherPortNames, BaseName);
-}
+}
\ No newline at end of file
diff --git a/Engine/Plugins/VirtualProduction/DMX/DMXProtocol/Source/DMXProtocol/Public/DMXProtocolSettings.h b/Engine/Plugins/VirtualProduction/DMX/DMXProtocol/Source/DMXProtocol/Public/DMXProtocolSettings.h
--- a/Engine/Plugins/VirtualProduction/DMX/DMXProtocol/Source/DMXProtocol/Public/DMXProtocolSettings.h
+++ b/Engine/Plugins/VirtualProduction/DMX/DMXProtocol/Source/DMXProtocol/Public/DMXProtocolSettings.h
@@ -98,6 +98,12 @@ public:
 	/** Overrides if send DMX is enabled at runtime */
 	void OverrideReceiveDMXEnabled(bool bEnabled);
 
+	/** Returns an automatic/unique name for an InputPort */
+	FString GetUniqueInputPortName() const;
+
+	/** Returns an automatic/unique name for an OutputPort */
+	FString GetUniqueOutputPortName() const;
+
 private:
 	/** Whether DMX is sent to the network. Recalled whenever editor or game starts.  */
 	UPROPERTY(Config, EditAnywhere, Category = "DMX|Communication Settings", Meta = (AllowPrivateAccess = true, DisplayName = "Send DMX by default"))
diff --git a/Engine/Plugins/VirtualProduction/DMX/DMXProtocol/Source/DMXProtocol/Public/IO/DMXInputPortConfig.h b/Engine/Plugins/VirtualProduction/DMX/DMXProtocol/Source/DMXProtocol/Public/IO/DMXInputPortConfig.h
--- a/Engine/Plugins/VirtualProduction/DMX/DMXProtocol/Source/DMXProtocol/Public/IO/DMXInputPortConfig.h
+++ b/Engine/Plugins/VirtualProduction/DMX/DMXProtocol/Source/DMXProtocol/Public/IO/DMXInputPortConfig.h
@@ -142,8 +142,6 @@ protected:
 	int32 Priority = 0;
 
 protected:
-	/** Generates a unique port name (unique for those stored in project settings) */
-	void GenerateUniquePortName();
 	
 	/** 
 	 * Unique identifier, shared with the port instance.
diff --git a/Engine/Plugins/VirtualProduction/DMX/DMXProtocol/Source/DMXProtocol/Public/IO/DMXOutputPortConfig.h b/Engine/Plugins/VirtualProduction/DMX/DMXProtocol/Source/DMXProtocol/Public/IO/DMXOutputPortConfig.h
--- a/Engine/Plugins/VirtualProduction/DMX/DMXProtocol/Source/DMXProtocol/Public/IO/DMXOutputPortConfig.h
+++ b/Engine/Plugins/VirtualProduction/DMX/DMXProtocol/Source/DMXProtocol/Public/IO/DMXOutputPortConfig.h
@@ -175,8 +175,6 @@ protected:
 	FFrameRate DelayFrameRate;
 
 protected:
-	/** Generates a unique port name (unique for those stored in project settings) */
-	void GenerateUniquePortName();
 
 	/** 
 	 * Unique identifier, shared with the port instance.
