Trivial fix for Linux build

diff --git a/Engine/Plugins/VirtualProduction/DMX/DMXProtocol/Source/DMXProtocol/Private/DMXProtocolSettings.cpp b/Engine/Plugins/VirtualProduction/DMX/DMXProtocol/Source/DMXProtocol/Private/DMXProtocolSettings.cpp
--- a/Engine/Plugins/VirtualProduction/DMX/DMXProtocol/Source/DMXProtocol/Private/DMXProtocolSettings.cpp
+++ b/Engine/Plugins/VirtualProduction/DMX/DMXProtocol/Source/DMXProtocol/Private/DMXProtocolSettings.cpp
@@ -258,7 +258,7 @@ FString UDMXProtocolSettings::GetUniqueInputPortName() const
 	const FString AutoNamePrefix = "InputPort";
 	int32 HighestPortNumber = 0;
 
-	for (const FDMXInputPortConfig InputPortConfig : InputPortConfigs)
+	for (const FDMXInputPortConfig& InputPortConfig : InputPortConfigs)
 	{
 		const FString& PortName = InputPortConfig.GetPortName();
 		if (PortName.StartsWith(AutoNamePrefix))
