DMX - Fix more issues with delay, an issue with concurrency when ports are assigned to their protocols and add various log messages to console commands

diff --git a/Engine/Plugins/VirtualProduction/DMX/DMXProtocol/Source/DMXProtocol/Private/IO/DMXInputPort.cpp b/Engine/Plugins/VirtualProduction/DMX/DMXProtocol/Source/DMXProtocol/Private/IO/DMXInputPort.cpp
--- a/Engine/Plugins/VirtualProduction/DMX/DMXProtocol/Source/DMXProtocol/Private/IO/DMXInputPort.cpp
+++ b/Engine/Plugins/VirtualProduction/DMX/DMXProtocol/Source/DMXProtocol/Private/IO/DMXInputPort.cpp
@@ -16,21 +16,82 @@ DECLARE_CYCLE_STAT(TEXT("Input Port Tick"), STAT_DMXInputPortTick, STATGROUP_DMX
 
 #define LOCTEXT_NAMESPACE "DMXInputPort"
 
+namespace UE::DMX::DMXInputPort
+{
+	/** Finds an Input port by its name, optionally logging when it can't be found */
+	static const FDMXInputPortSharedPtr FindInputPortByName(const FString& PortName, bool bPrintToLogIfNotFound = true)
+	{
+		const FDMXInputPortSharedRef* InputPortPtr = FDMXPortManager::Get().GetInputPorts().FindByPredicate([PortName](const FDMXInputPortSharedPtr& InputPort)
+			{
+				return InputPort->GetPortName() == PortName;
+			});
+
+		if (InputPortPtr)
+		{
+			return *InputPortPtr;
+		}
+		else if (bPrintToLogIfNotFound)
+		{
+			UE_LOG(LogDMXProtocol, Warning, TEXT("Could not find Input Port '%s'. Available Input Ports and their names can be found in in Project settings -> Plugins -> DMX."), *PortName);
+		}
+
+		return nullptr;
+	}
+
+	/** Tests the console command, returns true if they yield a valid port. Logs issues otherwise. Handles ? commands. */
+	static bool VerifyConsoleCommandArguments(const FString& ConsoleCommand, int32 MinNumExpectedArgs, const TArray<FString>& Args)
+	{
+		check(MinNumExpectedArgs >= 1);
+
+		if (Args.Num() < MinNumExpectedArgs)
+		{
+			UE_LOG(LogDMXProtocol, Warning, TEXT("Console Command %s failed. Insufficient arguments specified."), *ConsoleCommand);
+			UE_LOG(LogDMXProtocol, Warning, TEXT("Use '%s ?' for help. Use '%s PortName ?' to print the current configuration to logs."), *ConsoleCommand, *ConsoleCommand);
+			return false;
+		}
+
+		const FString& PortName = Args[0];
+		const FDMXInputPortSharedPtr InputPort = UE::DMX::DMXInputPort::FindInputPortByName(PortName);
+		if (InputPort.IsValid())
+		{
+			// Log the port if the console command sets the second arg to '?' (e.g. 'DMX.SetInputPortProtocol InputPortA ?')
+			if (Args.Num() == 2 && Args[1] == TEXT("?"))
+			{
+				UE_LOG(LogDMXProtocol, Log, TEXT("Input Port: %s"), *InputPort->GetPortName());
+
+				FDMXInputPortConfig PortConfig = InputPort->MakeInputPortConfig();
+				UE_LOG(LogDMXProtocol, Log, TEXT("Protocol: %s"), *PortConfig.GetProtocolName().ToString());
+				UE_LOG(LogDMXProtocol, Log, TEXT("Communication Type: %i"), static_cast<int32>(PortConfig.GetCommunicationType()));
+				UE_LOG(LogDMXProtocol, Log, TEXT("Device Address: %s"), *PortConfig.GetDeviceAddress());
+
+				UE_LOG(LogDMXProtocol, Log, TEXT("Local Universe Start: %i"), PortConfig.GetLocalUniverseStart());
+				UE_LOG(LogDMXProtocol, Log, TEXT("Num Universes: %i"), PortConfig.GetNumUniverses());
+				UE_LOG(LogDMXProtocol, Log, TEXT("Extern Universe Start: %i"), PortConfig.GetExternUniverseStart());
+				UE_LOG(LogDMXProtocol, Log, TEXT("Priority Strategy: %i"), static_cast<int32>(PortConfig.GetPortPriorityStrategy()));
+				UE_LOG(LogDMXProtocol, Log, TEXT("Priority:	%i"), PortConfig.GetPriority());
+			}
+			else
+			{
+				return true;
+			}
+		}
+
+		return false;
+	}
+}
+
 /** Helper to override a member variable of an Input Port */
 #define DMX_OVERRIDE_INPUTPORT_VAR(MemberName, PortName, Value) \
 { \
-	const FDMXInputPortSharedRef* InputPortPtr = FDMXPortManager::Get().GetInputPorts().FindByPredicate([PortName](const FDMXInputPortSharedPtr& InputPort) \
-		{ \
-			return InputPort->GetPortName() == PortName; \
-		}); \
-	if (InputPortPtr) \
+	const FDMXInputPortSharedPtr OverRideInputPortMacro_InputPort = UE::DMX::DMXInputPort::FindInputPortByName(PortName); \
+	if (OverRideInputPortMacro_InputPort.IsValid()) \
 	{ \
-		const FDMXInputPortConfig OldInputPortConfig = (*InputPortPtr)->MakeInputPortConfig(); \
+		const FDMXInputPortConfig OldInputPortConfig = OverRideInputPortMacro_InputPort->MakeInputPortConfig(); \
 		FDMXInputPortConfigParams InputPortConfigParams = FDMXInputPortConfigParams(OldInputPortConfig); \
 		InputPortConfigParams.MemberName = Value; \
-		FDMXInputPortConfig NewInputPortConfig((*InputPortPtr)->GetPortGuid(), InputPortConfigParams); \
+		FDMXInputPortConfig NewInputPortConfig(OverRideInputPortMacro_InputPort->GetPortGuid(), InputPortConfigParams); \
 		constexpr bool bForceUpdateRegistrationWithProtocol = true; \
-		(*InputPortPtr)->UpdateFromConfig(NewInputPortConfig, bForceUpdateRegistrationWithProtocol); \
+		OverRideInputPortMacro_InputPort->UpdateFromConfig(NewInputPortConfig, bForceUpdateRegistrationWithProtocol); \
 	} \
 }
 
@@ -40,7 +101,8 @@ static FAutoConsoleCommand GDMXSetInputPortProtocolCommand(
 	FConsoleCommandWithArgsDelegate::CreateStatic(
 		[](const TArray<FString>& Args)
 		{
-			if (Args.Num() < 2)
+			constexpr int32 MinNumExpectedArgs = 2;
+			if (!UE::DMX::DMXInputPort::VerifyConsoleCommandArguments(TEXT("DMX.SetInputPortProtocol"), MinNumExpectedArgs, Args))
 			{
 				return;
 			}
@@ -52,6 +114,14 @@ static FAutoConsoleCommand GDMXSetInputPortProtocolCommand(
 			{
 				DMX_OVERRIDE_INPUTPORT_VAR(ProtocolName, PortName, ProtocolNameValue);
 			}
+			else
+			{
+				UE_LOG(LogDMXProtocol, Warning, TEXT("Console command DMX.SetInputPortProtocol failed. Protocol '%s' does not exist or is not loaded. Available protocols are: "), *ProtocolNameValue.ToString());
+				for (const FName& ProtocolName : IDMXProtocol::GetProtocolNames())
+				{
+					UE_LOG(LogDMXProtocol, Warning, TEXT("%s"), *ProtocolName.ToString());
+				}
+			}
 		})
 );
 
@@ -61,7 +131,8 @@ static FAutoConsoleCommand GDMXSetInputPortDeviceAddressCommand(
 	FConsoleCommandWithArgsDelegate::CreateStatic(
 		[](const TArray<FString>& Args)
 		{
-			if (Args.Num() < 2)
+			constexpr int32 MinNumExpectedArgs = 2;
+			if (!UE::DMX::DMXInputPort::VerifyConsoleCommandArguments(TEXT("DMX.SetInputPortProtocol"), MinNumExpectedArgs, Args))
 			{
 				return;
 			}
@@ -79,19 +150,38 @@ static FAutoConsoleCommand GDMXSetInputPortLocalUniverseStartCommand(
 	FConsoleCommandWithArgsDelegate::CreateStatic(
 		[](const TArray<FString>& Args)
 		{
-			if (Args.Num() < 2)
+			constexpr int32 MinNumExpectedArgs = 2;
+			if (!UE::DMX::DMXInputPort::VerifyConsoleCommandArguments(TEXT("DMX.SetInputPortProtocol"), MinNumExpectedArgs, Args))
 			{
 				return;
 			}
 
 			const FString& PortName = Args[0];
-			const FString& LocalUniverseStartValueString = Args[1];
-			int32 LocalUniverseStartValue;
-			if (LexTryParseString<int32>(LocalUniverseStartValue, *LocalUniverseStartValueString))
+			const FDMXInputPortSharedPtr InputPort = UE::DMX::DMXInputPort::FindInputPortByName(PortName);
+
+			if (InputPort.IsValid())
 			{
-				if (LocalUniverseStartValue >= 0 && LocalUniverseStartValue <= DMX_MAX_UNIVERSE)
+				if (const IDMXProtocolPtr& Protocol = InputPort->GetProtocol())
 				{
-					DMX_OVERRIDE_INPUTPORT_VAR(LocalUniverseStart, PortName, LocalUniverseStartValue);
+					const FString& LocalUniverseStartValueString = Args[1];
+					int32 LocalUniverseStartValue;
+					if (LexTryParseString<int32>(LocalUniverseStartValue, *LocalUniverseStartValueString))
+					{
+						const int32 NumUniverses = InputPort->GetLocalUniverseEnd() - InputPort->GetLocalUniverseStart() + 1;
+
+						const int32 ExternUniverseStart = InputPort->ConvertLocalToExternUniverseID(LocalUniverseStartValue);
+						const int32 ExternUniverseEnd = InputPort->ConvertLocalToExternUniverseID(LocalUniverseStartValue + NumUniverses - 1);
+
+						if (Protocol->IsValidUniverseID(ExternUniverseStart) &&
+							Protocol->IsValidUniverseID(ExternUniverseEnd))
+						{
+							DMX_OVERRIDE_INPUTPORT_VAR(LocalUniverseStart, PortName, LocalUniverseStartValue);
+						}
+						else
+						{
+							UE_LOG(LogDMXProtocol, Warning, TEXT("Console command DMX.SetInputPortLocalUniverseStart failed. Local Universe Start '%s' along with Num Universes '%i' results in a Universe range that is not supported by the Protocol of the Port."), *LocalUniverseStartValueString, NumUniverses);
+						}
+					}
 				}
 			}
 		})
@@ -103,19 +193,39 @@ static FAutoConsoleCommand GDMXSetInputPortNumUniversesCommand(
 	FConsoleCommandWithArgsDelegate::CreateStatic(
 		[](const TArray<FString>& Args)
 		{
-			if (Args.Num() < 2)
+			constexpr int32 MinNumExpectedArgs = 2;
+			if (!UE::DMX::DMXInputPort::VerifyConsoleCommandArguments(TEXT("DMX.SetInputPortProtocol"), MinNumExpectedArgs, Args))
 			{
 				return;
 			}
 
 			const FString& PortName = Args[0];
-			const FString& NumUniversesValueString = Args[1];
-			int32 NumUniversesValue;
-			if (LexTryParseString<int32>(NumUniversesValue, *NumUniversesValueString))
+			const FDMXInputPortSharedPtr InputPort = UE::DMX::DMXInputPort::FindInputPortByName(PortName);
+
+			if (InputPort.IsValid())
 			{
-				if (NumUniversesValue >= 0 && NumUniversesValue <= DMX_MAX_UNIVERSE)
+				if (const IDMXProtocolPtr& Protocol = InputPort->GetProtocol())
 				{
-					DMX_OVERRIDE_INPUTPORT_VAR(NumUniverses, PortName, NumUniversesValue);
+					const FString& NumUniversesValueString = Args[1];
+					int32 NumUniversesValue;
+					if (LexTryParseString<int32>(NumUniversesValue, *NumUniversesValueString))
+					{
+						const int32 LocalUniverseStart = InputPort->GetLocalUniverseStart();
+
+						const int32 ExternUniverseStart = InputPort->ConvertLocalToExternUniverseID(LocalUniverseStart);
+						const int32 ExternUniverseEnd = ExternUniverseStart + NumUniversesValue - 1;
+
+						if (NumUniversesValue > 0 &&
+							Protocol->IsValidUniverseID(ExternUniverseStart) &&
+							Protocol->IsValidUniverseID(ExternUniverseEnd))
+						{
+							DMX_OVERRIDE_INPUTPORT_VAR(NumUniverses, PortName, NumUniversesValue);
+						}
+						else
+						{
+							UE_LOG(LogDMXProtocol, Warning, TEXT("Console command DMX.SetInputPortNumUniverses failed. Local Universe Start '%i' along with Num Universes '%s' results in a Universe range that is not supported by the Protocol of the Port."), LocalUniverseStart, *NumUniversesValueString);
+						}
+					}
 				}
 			}
 		})
@@ -127,19 +237,36 @@ static FAutoConsoleCommand GDMXSetInputPortExternUniverseStartCommand(
 	FConsoleCommandWithArgsDelegate::CreateStatic(
 		[](const TArray<FString>& Args)
 		{
-			if (Args.Num() < 2)
+			constexpr int32 MinNumExpectedArgs = 2;
+			if (!UE::DMX::DMXInputPort::VerifyConsoleCommandArguments(TEXT("DMX.SetInputPortProtocol"), MinNumExpectedArgs, Args))
 			{
 				return;
 			}
 
 			const FString& PortName = Args[0];
-			const FString& ExternUniverseStartValueString = Args[1];
-			int32 ExternUniverseStartValue;
-			if (LexTryParseString<int32>(ExternUniverseStartValue, *ExternUniverseStartValueString))
+			const FDMXInputPortSharedPtr InputPort = UE::DMX::DMXInputPort::FindInputPortByName(PortName);
+
+			if (InputPort.IsValid())
 			{
-				if (ExternUniverseStartValue >= 0 && ExternUniverseStartValue <= DMX_MAX_UNIVERSE)
+				const FString& ExternUniverseStartValueString = Args[1];
+				int32 ExternUniverseStartValue;
+				if (LexTryParseString<int32>(ExternUniverseStartValue, *ExternUniverseStartValueString))
 				{
-					DMX_OVERRIDE_INPUTPORT_VAR(ExternUniverseStart, PortName, ExternUniverseStartValue);
+					if (const IDMXProtocolPtr& Protocol = InputPort->GetProtocol())
+					{
+						const int32 NumUniverses = InputPort->GetLocalUniverseEnd() - InputPort->GetLocalUniverseStart() + 1;
+						const int32 ExternUniverseEnd = ExternUniverseStartValue + NumUniverses - 1;
+
+						if (Protocol->IsValidUniverseID(ExternUniverseStartValue) &&
+							Protocol->IsValidUniverseID(ExternUniverseEnd))
+						{
+							DMX_OVERRIDE_INPUTPORT_VAR(ExternUniverseStart, PortName, ExternUniverseStartValue);
+						}
+						else
+						{
+							UE_LOG(LogDMXProtocol, Warning, TEXT("Console command DMX.SetInputPortExternUniverseStart failed. Extern Universe Start '%s' along with Num Universes '%i' results in a Universe range that is not supported by the Protocol of the Port."), *ExternUniverseStartValueString, NumUniverses);
+						}
+					}
 				}
 			}
 		})
@@ -151,7 +278,8 @@ static FAutoConsoleCommand GDMXSetInputPortPriorityStrategyCommand(
 	FConsoleCommandWithArgsDelegate::CreateStatic(
 		[](const TArray<FString>& Args)
 		{
-			if (Args.Num() < 2)
+			constexpr int32 MinNumExpectedArgs = 2;
+			if (!UE::DMX::DMXInputPort::VerifyConsoleCommandArguments(TEXT("DMX.SetInputPortProtocol"), MinNumExpectedArgs, Args))
 			{
 				return;
 			}
@@ -182,7 +310,8 @@ static FAutoConsoleCommand GDMXSetInputPortPriorityCommand(
 	FConsoleCommandWithArgsDelegate::CreateStatic(
 		[](const TArray<FString>& Args)
 		{
-			if (Args.Num() < 2)
+			constexpr int32 MinNumExpectedArgs = 2;
+			if (!UE::DMX::DMXInputPort::VerifyConsoleCommandArguments(TEXT("DMX.SetInputPortProtocol"), MinNumExpectedArgs, Args))
 			{
 				return;
 			}
@@ -203,31 +332,29 @@ static FAutoConsoleCommand GDMXResetInputPortToProjectSettings(
 	FConsoleCommandWithArgsDelegate::CreateStatic(
 		[](const TArray<FString>& Args)
 		{
-			if (Args.Num() < 1)
+			constexpr int32 MinNumExpectedArgs = 1;
+			if (!UE::DMX::DMXInputPort::VerifyConsoleCommandArguments(TEXT("DMX.SetInputPortProtocol"), MinNumExpectedArgs, Args))
 			{
 				return;
 			}
 
 			const FString& PortName = Args[0];
-			const FDMXInputPortSharedRef* PortPtr = FDMXPortManager::Get().GetInputPorts().FindByPredicate([PortName](const FDMXInputPortSharedPtr& InputPort)
-				{
-					return InputPort->GetPortName() == PortName;
-				});
 
-			if (PortPtr)
+			const FDMXInputPortSharedPtr InputPort = UE::DMX::DMXInputPort::FindInputPortByName(PortName);
+			if(InputPort.IsValid())
 			{
 				const UDMXProtocolSettings* ProtocolSettings = GetDefault<UDMXProtocolSettings>();
 				if (ProtocolSettings)
 				{
-					const FDMXInputPortConfig* PortConfigPtr = ProtocolSettings->InputPortConfigs.FindByPredicate([PortPtr](const FDMXInputPortConfig& InputPortConfig)
+					const FDMXInputPortConfig* PortConfigPtr = ProtocolSettings->InputPortConfigs.FindByPredicate([InputPort](const FDMXInputPortConfig& InputPortConfig)
 						{
-							return InputPortConfig.GetPortGuid() == (*PortPtr)->GetPortGuid();
+							return InputPortConfig.GetPortGuid() == InputPort->GetPortGuid();
 						});
 
 					if (PortConfigPtr)
 					{
 						FDMXInputPortConfig PortConfig = *PortConfigPtr;
-						(*PortPtr)->UpdateFromConfig(PortConfig);
+						InputPort->UpdateFromConfig(PortConfig);
 					}
 				}
 			}
diff --git a/Engine/Plugins/VirtualProduction/DMX/DMXProtocol/Source/DMXProtocol/Private/IO/DMXOutputPort.cpp b/Engine/Plugins/VirtualProduction/DMX/DMXProtocol/Source/DMXProtocol/Private/IO/DMXOutputPort.cpp
--- a/Engine/Plugins/VirtualProduction/DMX/DMXProtocol/Source/DMXProtocol/Private/IO/DMXOutputPort.cpp
+++ b/Engine/Plugins/VirtualProduction/DMX/DMXProtocol/Source/DMXProtocol/Private/IO/DMXOutputPort.cpp
@@ -17,21 +17,94 @@
 
 #define LOCTEXT_NAMESPACE "DMXOutputPort"
 
+
+namespace UE::DMX::DMXOutputPort
+{
+	/** Finds an Output port by its name, optionally logging when it can't be found */
+	static const FDMXOutputPortSharedPtr FindOutputPortByName(const FString& PortName, bool bPrintToLogIfNotFound = true)
+	{
+		const FDMXOutputPortSharedRef* OutputPortPtr = FDMXPortManager::Get().GetOutputPorts().FindByPredicate([PortName](const FDMXOutputPortSharedPtr& OutputPort)
+			{
+				return OutputPort->GetPortName() == PortName;
+			});
+
+		if (OutputPortPtr)
+		{
+			return *OutputPortPtr;
+		}
+		else if (bPrintToLogIfNotFound)
+		{
+			UE_LOG(LogDMXProtocol, Warning, TEXT("Could not find Output Port '%s'.  Available Output Ports and their names can be found in Project settings -> Plugins -> DMX."), *PortName);
+		}
+
+		return nullptr;
+	}
+
+	/** Tests the console command, returns true if they yield a valid port. Logs issues otherwise. Handles ? commands. */
+	static bool VerifyConsoleCommandArguments(const FString& ConsoleCommand, int32 MinNumExpectedArgs, const TArray<FString>& Args)
+	{
+		check(MinNumExpectedArgs >= 1);
+
+		if (Args.Num() < MinNumExpectedArgs)
+		{
+			UE_LOG(LogDMXProtocol, Warning, TEXT("Console Command %s failed. Insufficient arguments specified."), *ConsoleCommand);
+			UE_LOG(LogDMXProtocol, Warning, TEXT("Use '%s ?' for help. Use '%s PortName ?' to print the current configuration to logs."), *ConsoleCommand, *ConsoleCommand);
+			return false;
+		}
+
+		const FString& PortName = Args[0];
+		const FDMXOutputPortSharedPtr OutputPort = UE::DMX::DMXOutputPort::FindOutputPortByName(PortName);
+		if (OutputPort.IsValid())
+		{
+			// Log the port if the console command sets the second arg to '?' (e.g. 'DMX.SetOutputPortProtocol OutputPortA ?')
+			if (Args.Num() == 2 && Args[1] == TEXT("?"))
+			{
+				// Log the port
+				UE_LOG(LogDMXProtocol, Log, TEXT("Output Port: %s"), *OutputPort->GetPortName());
+
+				FDMXOutputPortConfig PortConfig = OutputPort->MakeOutputPortConfig();
+				UE_LOG(LogDMXProtocol, Log, TEXT("Protocol: %s"), *PortConfig.GetProtocolName().ToString());
+				UE_LOG(LogDMXProtocol, Log, TEXT("Communication Type: %i"), static_cast<int32>(PortConfig.GetCommunicationType()));
+				UE_LOG(LogDMXProtocol, Log, TEXT("Device Address: %s"), *PortConfig.GetDeviceAddress());
+
+				const TArray<FDMXOutputPortDestinationAddress> DestinationAddresses = PortConfig.GetDestinationAddresses();
+				for (int32 IndexDestinationAddress = 0; IndexDestinationAddress < DestinationAddresses.Num(); IndexDestinationAddress++)
+				{
+					UE_LOG(LogDMXProtocol, Log, TEXT("Destination Address [%i]: %s"), IndexDestinationAddress, *DestinationAddresses[IndexDestinationAddress].DestinationAddressString);
+				}
+
+				UE_LOG(LogDMXProtocol, Log, TEXT("Input into Engine: %s"), PortConfig.NeedsLoopbackToEngine() ? TEXT("True") : TEXT("False"));
+				UE_LOG(LogDMXProtocol, Log, TEXT("Local Universe Start: %i"), PortConfig.GetLocalUniverseStart());
+				UE_LOG(LogDMXProtocol, Log, TEXT("Num Universes: %i"), PortConfig.GetNumUniverses());
+				UE_LOG(LogDMXProtocol, Log, TEXT("Extern Universe Start: %i"), PortConfig.GetExternUniverseStart());
+				UE_LOG(LogDMXProtocol, Log, TEXT("Priority:	%i"), PortConfig.GetPriority());
+
+				UE_LOG(LogDMXProtocol, Log, TEXT("Delay: %f"), PortConfig.GetDelay());
+				UE_LOG(LogDMXProtocol, Log, TEXT("Delay Frame Rate:	%i/%i"), PortConfig.GetDelayFrameRate().Numerator, PortConfig.GetDelayFrameRate().Denominator);
+				UE_LOG(LogDMXProtocol, Log, TEXT("Delay in Seconds: %f"), OutputPort->GetDelaySeconds());
+			}
+			else
+			{
+				return true;
+			}
+		}
+
+		return false;
+	}
+}
+
 /** Helper to override a member variable of an Output Port */
 #define DMX_OVERRIDE_OUTPUTPORT_VAR(MemberName, PortName, Value) \
 { \
-	const FDMXOutputPortSharedRef* OutputPortPtr = FDMXPortManager::Get().GetOutputPorts().FindByPredicate([PortName](const FDMXOutputPortSharedPtr& OutputPort) \
-		{ \
-			return OutputPort->GetPortName() == PortName; \
-		}); \
-	if (OutputPortPtr) \
+	const FDMXOutputPortSharedPtr OverRideOutputPortMacro_OutputPort = UE::DMX::DMXOutputPort::FindOutputPortByName(PortName); \
+	if (OverRideOutputPortMacro_OutputPort.IsValid()) \
 	{ \
-		const FDMXOutputPortConfig OldOutputPortConfig = (*OutputPortPtr)->MakeOutputPortConfig(); \
+		const FDMXOutputPortConfig OldOutputPortConfig = OverRideOutputPortMacro_OutputPort->MakeOutputPortConfig(); \
 		FDMXOutputPortConfigParams OutputPortConfigParams = FDMXOutputPortConfigParams(OldOutputPortConfig); \
 		OutputPortConfigParams.MemberName = Value; \
-		FDMXOutputPortConfig NewOutputPortConfig((*OutputPortPtr)->GetPortGuid(), OutputPortConfigParams); \
+		FDMXOutputPortConfig NewOutputPortConfig(OverRideOutputPortMacro_OutputPort->GetPortGuid(), OutputPortConfigParams); \
 		constexpr bool bForceUpdateRegistrationWithProtocol = true; \
-		(*OutputPortPtr)->UpdateFromConfig(NewOutputPortConfig, bForceUpdateRegistrationWithProtocol); \
+		OverRideOutputPortMacro_OutputPort->UpdateFromConfig(NewOutputPortConfig, bForceUpdateRegistrationWithProtocol); \
 	} \
 }
 
@@ -41,7 +114,8 @@ static FAutoConsoleCommand GDMXSetOutputPortProtocolCommand(
 	FConsoleCommandWithArgsDelegate::CreateStatic(
 		[](const TArray<FString>& Args)
 		{
-			if (Args.Num() < 2)
+			constexpr int32 MinNumExpectedArgs = 2;
+			if (!UE::DMX::DMXOutputPort::VerifyConsoleCommandArguments(TEXT("DMX.SetOutputPortProtocol"), MinNumExpectedArgs, Args))
 			{
 				return;
 			}
@@ -53,6 +127,14 @@ static FAutoConsoleCommand GDMXSetOutputPortProtocolCommand(
 			{
 				DMX_OVERRIDE_OUTPUTPORT_VAR(ProtocolName, PortName, ProtocolNameValue);
 			}
+			else
+			{
+				UE_LOG(LogDMXProtocol, Warning, TEXT("Console command DMX.SetOutputPortProtocol failed. Protocol '%s' does not exist or is not loaded. Available protocols are: "), *ProtocolNameValue.ToString());
+				for (const FName& ProtocolName : IDMXProtocol::GetProtocolNames())
+				{
+					UE_LOG(LogDMXProtocol, Warning, TEXT("%s"), *ProtocolName.ToString());
+				}
+			}
 		})
 );
 
@@ -62,13 +144,15 @@ static FAutoConsoleCommand GDMXSetOutputPortCommunicationTypeCommand(
 	FConsoleCommandWithArgsDelegate::CreateStatic(
 		[](const TArray<FString>& Args)
 		{
-			if (Args.Num() < 2)
+			constexpr int32 MinNumExpectedArgs = 2;
+			if (!UE::DMX::DMXOutputPort::VerifyConsoleCommandArguments(TEXT("DMX.SetOutputPortCommunicationType"), MinNumExpectedArgs, Args))
 			{
 				return;
 			}
 
 			const FString& PortName = Args[0];
 			const FString& CommunicationTypeValueString = Args[1];
+
 			uint8 CommunicationTypeValue;
 			if (LexTryParseString<uint8>(CommunicationTypeValue, *CommunicationTypeValueString))
 			{
@@ -77,14 +161,11 @@ static FAutoConsoleCommand GDMXSetOutputPortCommunicationTypeCommand(
 					CommunicationTypeEnumValue == EDMXCommunicationType::Multicast ||
 					CommunicationTypeEnumValue == EDMXCommunicationType::Unicast)
 				{
-					const FDMXOutputPortSharedRef* PortPtr = FDMXPortManager::Get().GetOutputPorts().FindByPredicate([PortName](const FDMXOutputPortSharedPtr& OutputPort)
-						{
-							return OutputPort->GetPortName() == PortName;
-						});
+					const FDMXOutputPortSharedPtr OutputPort = UE::DMX::DMXOutputPort::FindOutputPortByName(PortName);
 
-					if (PortPtr)
+					if (OutputPort.IsValid())
 					{
-						if (const IDMXProtocolPtr Protocol = (*PortPtr)->GetProtocol())
+						if (const IDMXProtocolPtr Protocol = OutputPort->GetProtocol())
 						{
 							if (Protocol->GetOutputPortCommunicationTypes().Contains(CommunicationTypeEnumValue))
 							{
@@ -104,7 +185,8 @@ static FAutoConsoleCommand GDMXSetOutputPortOutputPortDeviceAddressCommand(
 	FConsoleCommandWithArgsDelegate::CreateStatic(
 		[](const TArray<FString>& Args)
 		{
-			if (Args.Num() < 2)
+			constexpr int32 MinNumExpectedArgs = 2;
+			if (!UE::DMX::DMXOutputPort::VerifyConsoleCommandArguments(TEXT("DMX.SetOutputPortDeviceAddress"), MinNumExpectedArgs, Args))
 			{
 				return;
 			}
@@ -122,7 +204,8 @@ static FAutoConsoleCommand GDMXSetOutputPortOutputPortDestinationAddressesComman
 	FConsoleCommandWithArgsDelegate::CreateStatic(
 		[](const TArray<FString>& Args)
 		{
-			if (Args.Num() < 2)
+			constexpr int32 MinNumExpectedArgs = 2;
+			if (!UE::DMX::DMXOutputPort::VerifyConsoleCommandArguments(TEXT("DMX.SetOutputPortDestinationAddresses"), MinNumExpectedArgs, Args))
 			{
 				return;
 			}
@@ -145,7 +228,8 @@ static FAutoConsoleCommand GDMXSetOutputPortLoopbackToEngineCommand(
 	FConsoleCommandWithArgsDelegate::CreateStatic(
 		[](const TArray<FString>& Args)
 		{
-			if (Args.Num() < 2)
+			constexpr int32 MinNumExpectedArgs = 2;
+			if (!UE::DMX::DMXOutputPort::VerifyConsoleCommandArguments(TEXT("DMX.SetOutputPortInputIntoEngine"), MinNumExpectedArgs, Args))
 			{
 				return;
 			}
@@ -163,19 +247,38 @@ static FAutoConsoleCommand GDMXSetOutputPortLocalUniverseStartCommand(
 	FConsoleCommandWithArgsDelegate::CreateStatic(
 		[](const TArray<FString>& Args)
 		{
-			if (Args.Num() < 2)
+			constexpr int32 MinNumExpectedArgs = 2;
+			if (!UE::DMX::DMXOutputPort::VerifyConsoleCommandArguments(TEXT("DMX.SetOutputPortLocalUniverseStart"), MinNumExpectedArgs, Args))
 			{
 				return;
 			}
 
 			const FString& PortName = Args[0];
-			const FString& LocalUniverseStartValueString = Args[1];
-			int32 LocalUniverseStartValue;
-			if (LexTryParseString<int32>(LocalUniverseStartValue, *LocalUniverseStartValueString))
+			const FDMXOutputPortSharedPtr OutputPort = UE::DMX::DMXOutputPort::FindOutputPortByName(PortName);
+
+			if (OutputPort.IsValid())
 			{
-				if (LocalUniverseStartValue >= 0 && LocalUniverseStartValue <= DMX_MAX_UNIVERSE)
+				if (const IDMXProtocolPtr& Protocol = OutputPort->GetProtocol())
 				{
-					DMX_OVERRIDE_OUTPUTPORT_VAR(LocalUniverseStart, PortName, LocalUniverseStartValue);
+					const FString& LocalUniverseStartValueString = Args[1];
+					int32 LocalUniverseStartValue;
+					if (LexTryParseString<int32>(LocalUniverseStartValue, *LocalUniverseStartValueString))
+					{
+						const int32 NumUniverses = OutputPort->GetLocalUniverseEnd() - OutputPort->GetLocalUniverseStart() + 1;
+
+						const int32 ExternUniverseStart = OutputPort->ConvertLocalToExternUniverseID(LocalUniverseStartValue);
+						const int32 ExternUniverseEnd = OutputPort->ConvertLocalToExternUniverseID(LocalUniverseStartValue + NumUniverses - 1);
+
+						if (Protocol->IsValidUniverseID(ExternUniverseStart) &&
+							Protocol->IsValidUniverseID(ExternUniverseEnd))
+						{
+							DMX_OVERRIDE_OUTPUTPORT_VAR(LocalUniverseStart, PortName, LocalUniverseStartValue);
+						}
+						else
+						{
+							UE_LOG(LogDMXProtocol, Warning, TEXT("Console command DMX.SetOutputPortLocalUniverseStart failed. Local Universe Start '%s' along with Num Universes '%i' results in a Universe range that is not supported by the Protocol of the Port."), *LocalUniverseStartValueString, NumUniverses);
+						}
+					}
 				}
 			}
 		})
@@ -187,19 +290,39 @@ static FAutoConsoleCommand GDMXSetOutputPortNumUniversesCommand(
 	FConsoleCommandWithArgsDelegate::CreateStatic(
 		[](const TArray<FString>& Args)
 		{
-			if (Args.Num() < 2)
+			constexpr int32 MinNumExpectedArgs = 2;
+			if (!UE::DMX::DMXOutputPort::VerifyConsoleCommandArguments(TEXT("DMX.SetOutputPortNumUniverses"), MinNumExpectedArgs, Args))
 			{
 				return;
 			}
 
 			const FString& PortName = Args[0];
-			const FString& NumUniversesValueString = Args[1];
-			int32 NumUniversesValue;
-			if (LexTryParseString<int32>(NumUniversesValue, *NumUniversesValueString))
+			const FDMXOutputPortSharedPtr OutputPort = UE::DMX::DMXOutputPort::FindOutputPortByName(PortName);
+
+			if (OutputPort.IsValid())
 			{
-				if (NumUniversesValue >= 0 && NumUniversesValue <= DMX_MAX_UNIVERSE)
+				if (const IDMXProtocolPtr& Protocol = OutputPort->GetProtocol())
 				{
-					DMX_OVERRIDE_OUTPUTPORT_VAR(NumUniverses, PortName, NumUniversesValue);
+					const FString& NumUniversesValueString = Args[1];
+					int32 NumUniversesValue;
+					if (LexTryParseString<int32>(NumUniversesValue, *NumUniversesValueString))
+					{
+						const int32 LocalUniverseStart = OutputPort->GetLocalUniverseStart();
+
+						const int32 ExternUniverseStart = OutputPort->ConvertLocalToExternUniverseID(LocalUniverseStart);
+						const int32 ExternUniverseEnd = ExternUniverseStart + NumUniversesValue - 1;
+
+						if (NumUniversesValue > 0 &&
+							Protocol->IsValidUniverseID(ExternUniverseStart) &&
+							Protocol->IsValidUniverseID(ExternUniverseEnd))
+						{
+							DMX_OVERRIDE_OUTPUTPORT_VAR(NumUniverses, PortName, NumUniversesValue);
+						}
+						else
+						{
+							UE_LOG(LogDMXProtocol, Warning, TEXT("Console command DMX.SetOutputPortNumUniverses failed. Local Universe Start '%i' along with Num Universes '%s' results in a Universe range that is not supported by the Protocol of the Port."), LocalUniverseStart, *NumUniversesValueString);
+						}
+					}
 				}
 			}
 		})
@@ -211,19 +334,37 @@ static FAutoConsoleCommand GDMXSetOutputPortExternUniverseStartCommand(
 	FConsoleCommandWithArgsDelegate::CreateStatic(
 		[](const TArray<FString>& Args)
 		{
-			if (Args.Num() < 2)
+			constexpr int32 MinNumExpectedArgs = 2;
+			if (!UE::DMX::DMXOutputPort::VerifyConsoleCommandArguments(TEXT("DMX.SetOutputPortExternUniverseStart"), MinNumExpectedArgs, Args))
 			{
 				return;
 			}
 
 			const FString& PortName = Args[0];
-			const FString& ExternUniverseStartValueString = Args[1];
-			int32 ExternUniverseStartValue;
-			if (LexTryParseString<int32>(ExternUniverseStartValue, *ExternUniverseStartValueString))
+			const FDMXOutputPortSharedPtr OutputPort = UE::DMX::DMXOutputPort::FindOutputPortByName(PortName);
+
+
+			if (OutputPort.IsValid())
 			{
-				if (ExternUniverseStartValue >= 0 && ExternUniverseStartValue <= DMX_MAX_UNIVERSE)
+				if (const IDMXProtocolPtr& Protocol = OutputPort->GetProtocol())
 				{
-					DMX_OVERRIDE_OUTPUTPORT_VAR(ExternUniverseStart, PortName, ExternUniverseStartValue);
+					const FString& ExternUniverseStartValueString = Args[1];
+					int32 ExternUniverseStartValue;
+					if (LexTryParseString<int32>(ExternUniverseStartValue, *ExternUniverseStartValueString))
+					{
+						const int32 NumUniverses = OutputPort->GetLocalUniverseEnd() - OutputPort->GetLocalUniverseStart() + 1;
+						const int32 ExternUniverseEnd = ExternUniverseStartValue + NumUniverses - 1;
+
+						if (Protocol->IsValidUniverseID(ExternUniverseStartValue) &&
+							Protocol->IsValidUniverseID(ExternUniverseEnd))
+						{
+							DMX_OVERRIDE_OUTPUTPORT_VAR(ExternUniverseStart, PortName, ExternUniverseStartValue);
+						}
+						else
+						{
+							UE_LOG(LogDMXProtocol, Warning, TEXT("Console command DMX.SetOutputPortExternUniverseStart failed. Extern Universe Start '%s' along with Num Universes '%i' results in a Universe range that is not supported by the Protocol of the Port."), *ExternUniverseStartValueString, NumUniverses);
+						}
+					}
 				}
 			}
 		})
@@ -235,7 +376,8 @@ static FAutoConsoleCommand GDMXSetOutputPortPriorityCommand(
 	FConsoleCommandWithArgsDelegate::CreateStatic(
 		[](const TArray<FString>& Args)
 		{
-			if (Args.Num() < 2)
+			constexpr int32 MinNumExpectedArgs = 2;
+			if (!UE::DMX::DMXOutputPort::VerifyConsoleCommandArguments(TEXT("DMX.SetOutputPortPriority"), MinNumExpectedArgs, Args))
 			{
 				return;
 			}
@@ -256,7 +398,8 @@ static FAutoConsoleCommand GDMXSetOutputPortDelayCommand(
 	FConsoleCommandWithArgsDelegate::CreateStatic(
 		[](const TArray<FString>& Args)
 		{
-			if (Args.Num() < 2)
+			constexpr int32 MinNumExpectedArgs = 2;
+			if (!UE::DMX::DMXOutputPort::VerifyConsoleCommandArguments(TEXT("DMX.SetOutputPortDelay"), MinNumExpectedArgs, Args))
 			{
 				return;
 			}
@@ -267,34 +410,31 @@ static FAutoConsoleCommand GDMXSetOutputPortDelayCommand(
 			if (LexTryParseString<double>(DelayValue, *DelayValueString))
 			{
 				FFrameRate NewFrameRate = FFrameRate(1, 1);
-				
+
 				// Parse optional frame rate argument
 				if (Args.Num() == 3)
 				{
 					const FString& DelayFrameRateString = Args[2];
-				
+
 					TValueOrError<FFrameRate, FExpressionError> ParseResult = ParseFrameRate(*DelayFrameRateString);
 					if (ParseResult.IsValid())
 					{
 						NewFrameRate = ParseResult.GetValue();
 					}
+					else
+					{
+						UE_LOG(LogDMXProtocol, Warning, TEXT("Could not parse frame rate from console command DMX.SetOutputPortDelay. Note, not all decimal numbers are supported. E.g. 0.345s needs to be written as 345/1000."))
+					}
 				}
 
 				if (DelayValue >= 0.0)
 				{
-					const FDMXOutputPortSharedRef* PortPtr = FDMXPortManager::Get().GetOutputPorts().FindByPredicate([PortName](const FDMXOutputPortSharedPtr& OutputPort)
-						{
-							return OutputPort->GetPortName() == PortName;
-						});
+					const FDMXOutputPortSharedPtr OutputPort = UE::DMX::DMXOutputPort::FindOutputPortByName(PortName);
 
-					if (PortPtr)
+					if (OutputPort.IsValid())
 					{
-						// Limit to 60 seconds
-						const double FrameRateDecimal = static_cast<double>(NewFrameRate.Numerator) / NewFrameRate.Denominator;
-						const double NewDelaySeconds = FMath::Min(DelayValue / FrameRateDecimal, 60.f * FrameRateDecimal);
-
 						DMX_OVERRIDE_OUTPUTPORT_VAR(DelayFrameRate, PortName, NewFrameRate);
-						DMX_OVERRIDE_OUTPUTPORT_VAR(DelaySeconds, PortName, NewDelaySeconds);
+						DMX_OVERRIDE_OUTPUTPORT_VAR(Delay, PortName, DelayValue);
 					}
 				}
 			}
@@ -307,31 +447,29 @@ static FAutoConsoleCommand GDMXResetOutputPortToProjectSettings(
 	FConsoleCommandWithArgsDelegate::CreateStatic(
 		[](const TArray<FString>& Args)
 		{
-			if (Args.Num() < 1)
+			constexpr int32 MinNumExpectedArgs = 1;
+			if (!UE::DMX::DMXOutputPort::VerifyConsoleCommandArguments(TEXT("DMX.ResetOutputPortToProjectSettings"), MinNumExpectedArgs, Args))
 			{
 				return;
 			}
 
 			const FString& PortName = Args[0];
-			const FDMXOutputPortSharedRef* PortPtr = FDMXPortManager::Get().GetOutputPorts().FindByPredicate([PortName](const FDMXOutputPortSharedPtr& OutputPort)
-				{
-					return OutputPort->GetPortName() == PortName;
-				});
+			const FDMXOutputPortSharedPtr OutputPort = UE::DMX::DMXOutputPort::FindOutputPortByName(PortName);
 
-			if (PortPtr)
+			if(OutputPort.IsValid())
 			{
 				const UDMXProtocolSettings* ProtocolSettings = GetDefault<UDMXProtocolSettings>();
 				if (ProtocolSettings)
 				{
-					const FDMXOutputPortConfig* PortConfigPtr = ProtocolSettings->OutputPortConfigs.FindByPredicate([PortPtr](const FDMXOutputPortConfig& OutputPortConfig)
+					const FDMXOutputPortConfig* PortConfigPtr = ProtocolSettings->OutputPortConfigs.FindByPredicate([OutputPort](const FDMXOutputPortConfig& OutputPortConfig)
 						{
-							return OutputPortConfig.GetPortGuid() == (*PortPtr)->GetPortGuid();
+							return OutputPortConfig.GetPortGuid() == OutputPort->GetPortGuid();
 						});
 
 					if (PortConfigPtr)
 					{
 						FDMXOutputPortConfig PortConfig = *PortConfigPtr;
-						(*PortPtr)->UpdateFromConfig(PortConfig);
+						OutputPort->UpdateFromConfig(PortConfig);
 					}
 				}
 			}
@@ -399,7 +537,7 @@ FDMXOutputPortConfig FDMXOutputPort::MakeOutputPortConfig() const
 	Params.NumUniverses = NumUniverses;
 	Params.ExternUniverseStart = ExternUniverseStart;
 	Params.Priority = Priority;
-	Params.DelaySeconds = DelaySeconds;
+	Params.Delay = DelaySeconds * DelayFrameRate.AsDecimal();
 	Params.DelayFrameRate = DelayFrameRate;
 
 	return FDMXOutputPortConfig(PortGuid, Params);
@@ -440,7 +578,7 @@ void FDMXOutputPort::UpdateFromConfig(FDMXOutputPortConfig& InOutOutputPortConfi
 			DestinationAddresses == OutputPortConfig.GetDestinationAddresses() &&
 			CommunicationType == OutputPortConfig.GetCommunicationType() &&
 			Priority == OutputPortConfig.GetPriority() &&
-			DelaySeconds == OutputPortConfig.GetDelaySeconds())
+			DelaySeconds == OutputPortConfig.GetDelay() * OutputPortConfig.GetDelayFrameRate().AsDecimal())
 		{
 			return false;
 		}	
@@ -472,7 +610,7 @@ void FDMXOutputPort::UpdateFromConfig(FDMXOutputPortConfig& InOutOutputPortConfi
 	NumUniverses = OutputPortConfig.GetNumUniverses();
 	PortName = OutputPortConfig.GetPortName();
 	Priority = OutputPortConfig.GetPriority();
-	DelaySeconds = OutputPortConfig.GetDelaySeconds();
+	DelaySeconds = OutputPortConfig.GetDelay() / OutputPortConfig.GetDelayFrameRate().AsDecimal();
 	DelayFrameRate = OutputPortConfig.GetDelayFrameRate();
 
 	CommunicationDeterminator.SetLoopbackToEngine(OutputPortConfig.NeedsLoopbackToEngine());
diff --git a/Engine/Plugins/VirtualProduction/DMX/DMXProtocol/Source/DMXProtocol/Private/IO/DMXOutputPortConfig.cpp b/Engine/Plugins/VirtualProduction/DMX/DMXProtocol/Source/DMXProtocol/Private/IO/DMXOutputPortConfig.cpp
--- a/Engine/Plugins/VirtualProduction/DMX/DMXProtocol/Source/DMXProtocol/Private/IO/DMXOutputPortConfig.cpp
+++ b/Engine/Plugins/VirtualProduction/DMX/DMXProtocol/Source/DMXProtocol/Private/IO/DMXOutputPortConfig.cpp
@@ -21,7 +21,7 @@ FDMXOutputPortConfigParams::FDMXOutputPortConfigParams(const FDMXOutputPortConfi
 	, NumUniverses(OutputPortConfig.GetNumUniverses())
 	, ExternUniverseStart(OutputPortConfig.GetExternUniverseStart())
 	, Priority(OutputPortConfig.GetPriority())
-	, DelaySeconds(OutputPortConfig.GetDelaySeconds())
+	, Delay(OutputPortConfig.GetDelay())
 	, DelayFrameRate(OutputPortConfig.GetDelayFrameRate())
 {}
 
@@ -53,12 +53,13 @@ FDMXOutputPortConfig::FDMXOutputPortConfig(const FGuid& InPortGuid, const FDMXOu
 	, NumUniverses(InitializationData.NumUniverses)
 	, ExternUniverseStart(InitializationData.ExternUniverseStart)
 	, Priority(InitializationData.Priority)
-	, Delay(InitializationData.DelaySeconds * InitializationData.DelayFrameRate.AsDecimal())
+	, Delay(InitializationData.Delay)
 	, DelayFrameRate(InitializationData.DelayFrameRate)
 	, PortGuid(InPortGuid)
 {
 	// Cannot create port configs before the protocol module is up (it is required to sanetize protocol names).
 	check(FModuleManager::Get().IsModuleLoaded("DMXProtocol"));
+
 	check(PortGuid.IsValid());
 	check(!ProtocolName.IsNone())
 
@@ -139,11 +140,6 @@ void FDMXOutputPortConfig::MakeValid()
 	}
 }
 
-int32 FDMXOutputPortConfig::GetDelaySeconds() const
-{
-	return Delay / DelayFrameRate.AsDecimal();
-}
-
 FString FDMXOutputPortConfig::GetDeviceAddress() const
 {
 	// Allow to override the source ip from commandline
diff --git a/Engine/Plugins/VirtualProduction/DMX/DMXProtocol/Source/DMXProtocol/Public/IO/DMXOutputPort.h b/Engine/Plugins/VirtualProduction/DMX/DMXProtocol/Source/DMXProtocol/Public/IO/DMXOutputPort.h
--- a/Engine/Plugins/VirtualProduction/DMX/DMXProtocol/Source/DMXProtocol/Public/IO/DMXOutputPort.h
+++ b/Engine/Plugins/VirtualProduction/DMX/DMXProtocol/Source/DMXProtocol/Public/IO/DMXOutputPort.h
@@ -158,6 +158,9 @@ public:
 	UE_DEPRECATED(5.0, "Output Ports now support many destination addresses. Please use FDMXOutputPort::GetDestinationAddresses instead.")
 	FString GetDestinationAddress() const;
 
+	/** Returns the output port's delay in seconds */
+	FORCEINLINE double GetDelaySeconds() const { return DelaySeconds; }
+
 public:
 	//~ Begin DMXPort Interface 
 	virtual bool IsRegistered() const override;
diff --git a/Engine/Plugins/VirtualProduction/DMX/DMXProtocol/Source/DMXProtocol/Public/IO/DMXOutputPortConfig.h b/Engine/Plugins/VirtualProduction/DMX/DMXProtocol/Source/DMXProtocol/Public/IO/DMXOutputPortConfig.h
--- a/Engine/Plugins/VirtualProduction/DMX/DMXProtocol/Source/DMXProtocol/Public/IO/DMXOutputPortConfig.h
+++ b/Engine/Plugins/VirtualProduction/DMX/DMXProtocol/Source/DMXProtocol/Public/IO/DMXOutputPortConfig.h
@@ -53,7 +53,7 @@ struct DMXPROTOCOL_API FDMXOutputPortConfigParams
 	int32 NumUniverses;
 	int32 ExternUniverseStart;
 	int32 Priority;
-	double DelaySeconds;
+	double Delay;
 	FFrameRate DelayFrameRate;
 
 	// DEPRECATED 5.0. Instead please use the DestinationAddresses array.
@@ -94,7 +94,7 @@ public:
 	FORCEINLINE int32 GetExternUniverseStart() const { return ExternUniverseStart; }
 	FORCEINLINE int32 GetPriority() const { return Priority; }
 	FORCEINLINE const FGuid& GetPortGuid() const { return PortGuid; }
-	int32 GetDelaySeconds() const;
+	FORCEINLINE double GetDelay() const { return Delay; }
 	FORCEINLINE const FFrameRate& GetDelayFrameRate() const { return DelayFrameRate; }
 
 	UE_DEPRECATED(5.0, "Output Ports now support many destination addresses. Use FDMXOutputPortConfig::GetDestinationAddresses instead. ")
diff --git a/Engine/Plugins/VirtualProduction/DMX/DMXProtocol/Source/DMXProtocolArtNet/Private/DMXProtocolArtNetReceiver.cpp b/Engine/Plugins/VirtualProduction/DMX/DMXProtocol/Source/DMXProtocolArtNet/Private/DMXProtocolArtNetReceiver.cpp
--- a/Engine/Plugins/VirtualProduction/DMX/DMXProtocol/Source/DMXProtocolArtNet/Private/DMXProtocolArtNetReceiver.cpp
+++ b/Engine/Plugins/VirtualProduction/DMX/DMXProtocol/Source/DMXProtocolArtNet/Private/DMXProtocolArtNetReceiver.cpp
@@ -95,12 +95,16 @@ bool FDMXProtocolArtNetReceiver::EqualsEndpoint(const FString& IPAddress) const
 void FDMXProtocolArtNetReceiver::AssignInputPort(const TSharedPtr<FDMXInputPort, ESPMode::ThreadSafe>& InputPort)
 {
 	check(!AssignedInputPorts.Contains(InputPort));
+
+	const FScopeLock ChangeAssignedInputPortsLock(&ChangeAssignedInputPortsCriticalSection);
 	AssignedInputPorts.Add(InputPort);
 }
 
 void FDMXProtocolArtNetReceiver::UnassignInputPort(const TSharedPtr<FDMXInputPort, ESPMode::ThreadSafe>& InputPort)
 {
 	check(AssignedInputPorts.Contains(InputPort));
+
+	const FScopeLock ChangeAssignedInputPortsLock(&ChangeAssignedInputPortsCriticalSection);
 	AssignedInputPorts.Remove(InputPort);
 }
 
@@ -147,6 +151,8 @@ void FDMXProtocolArtNetReceiver::Update(const FTimespan& SocketWaitTime)
 		return;
 	}
 
+	const FScopeLock ChangeAssignedInputPortsLock(&ChangeAssignedInputPortsCriticalSection);
+
 	uint32 Size = 0;
 	int32 Read = 0;
 	while (Socket->HasPendingData(Size))
diff --git a/Engine/Plugins/VirtualProduction/DMX/DMXProtocol/Source/DMXProtocolArtNet/Private/DMXProtocolArtNetReceiver.h b/Engine/Plugins/VirtualProduction/DMX/DMXProtocol/Source/DMXProtocolArtNet/Private/DMXProtocolArtNetReceiver.h
--- a/Engine/Plugins/VirtualProduction/DMX/DMXProtocol/Source/DMXProtocolArtNet/Private/DMXProtocolArtNetReceiver.h
+++ b/Engine/Plugins/VirtualProduction/DMX/DMXProtocol/Source/DMXProtocolArtNet/Private/DMXProtocolArtNetReceiver.h
@@ -6,6 +6,7 @@
 
 #include "CoreMinimal.h"
 #include "HAL/Runnable.h"
+#include "Misc/ScopeLock.h"
 #include "Misc/SingleThreadRunnable.h"
 #include "Templates/Atomic.h"
 
@@ -112,6 +113,9 @@ private:
 	/** The sender when packets were received */
 	TSharedPtr<FInternetAddr> ReceivedSenderInternetAddr;
 
+	/** Critical section to be used when assigned input ports are changed */
+	FCriticalSection ChangeAssignedInputPortsCriticalSection;
+
 	/** Flag indicating that the thread is stopping. */
 	TAtomic<bool> bStopping;
 
diff --git a/Engine/Plugins/VirtualProduction/DMX/DMXProtocol/Source/DMXProtocolSACN/Private/DMXProtocolSACNReceiver.cpp b/Engine/Plugins/VirtualProduction/DMX/DMXProtocol/Source/DMXProtocolSACN/Private/DMXProtocolSACNReceiver.cpp
--- a/Engine/Plugins/VirtualProduction/DMX/DMXProtocol/Source/DMXProtocolSACN/Private/DMXProtocolSACNReceiver.cpp
+++ b/Engine/Plugins/VirtualProduction/DMX/DMXProtocol/Source/DMXProtocolSACN/Private/DMXProtocolSACNReceiver.cpp
@@ -109,6 +109,8 @@ bool FDMXProtocolSACNReceiver::EqualsEndpoint(const FString& IPAddress) const
 void FDMXProtocolSACNReceiver::AssignInputPort(const TSharedPtr<FDMXInputPort, ESPMode::ThreadSafe>& InputPort)
 {
 	check(!AssignedInputPorts.Contains(InputPort));
+
+	const FScopeLock ChangeAssignedInputPortsLock(&ChangeAssignedInputPortsCriticalSection);
 	AssignedInputPorts.Add(InputPort);
 
 	// Join the multicast groups for the ports where needed
@@ -140,6 +142,8 @@ void FDMXProtocolSACNReceiver::AssignInputPort(const TSharedPtr<FDMXInputPort, E
 void FDMXProtocolSACNReceiver::UnassignInputPort(const TSharedPtr<FDMXInputPort, ESPMode::ThreadSafe>& InputPort)
 {
 	check(AssignedInputPorts.Contains(InputPort));
+
+	const FScopeLock ChangeAssignedInputPortsLock(&ChangeAssignedInputPortsCriticalSection);
 	AssignedInputPorts.Remove(InputPort);
 
 	// Leave multicast groups outside of remaining port's universe ranges
@@ -209,6 +213,8 @@ void FDMXProtocolSACNReceiver::Update(const FTimespan& SocketWaitTime)
 		return;
 	}
 
+	const FScopeLock ChangeAssignedInputPortsLock(&ChangeAssignedInputPortsCriticalSection);
+
 	uint32 Size = 0;
 	int32 NumBytesRead = 0;
 	while (Socket->HasPendingData(Size))
diff --git a/Engine/Plugins/VirtualProduction/DMX/DMXProtocol/Source/DMXProtocolSACN/Private/DMXProtocolSACNReceiver.h b/Engine/Plugins/VirtualProduction/DMX/DMXProtocol/Source/DMXProtocolSACN/Private/DMXProtocolSACNReceiver.h
--- a/Engine/Plugins/VirtualProduction/DMX/DMXProtocol/Source/DMXProtocolSACN/Private/DMXProtocolSACNReceiver.h
+++ b/Engine/Plugins/VirtualProduction/DMX/DMXProtocol/Source/DMXProtocolSACN/Private/DMXProtocolSACNReceiver.h
@@ -6,6 +6,7 @@
 
 #include "CoreMinimal.h"
 #include "HAL/Runnable.h"
+#include "Misc/ScopeLock.h"
 #include "Misc/SingleThreadRunnable.h"
 #include "Templates/Atomic.h"
 
@@ -112,6 +113,9 @@ private:
 	/** The destination when packets are received */
 	TSharedPtr<FInternetAddr> ReceivedDestinationInternetAddr;
 
+	/** Critical section to be used when assigned input ports are changed */
+	FCriticalSection ChangeAssignedInputPortsCriticalSection;
+
 	/** Flag indicating that the thread is stopping. */
 	TAtomic<bool> bStopping;
 
