DMX - Add default attributes for ARRI fixtures

diff --git a/Engine/Plugins/VirtualProduction/DMX/DMXProtocol/Source/DMXProtocol/Private/DMXProtocolSettings.cpp b/Engine/Plugins/VirtualProduction/DMX/DMXProtocol/Source/DMXProtocol/Private/DMXProtocolSettings.cpp
--- a/Engine/Plugins/VirtualProduction/DMX/DMXProtocol/Source/DMXProtocol/Private/DMXProtocolSettings.cpp
+++ b/Engine/Plugins/VirtualProduction/DMX/DMXProtocol/Source/DMXProtocol/Private/DMXProtocolSettings.cpp
@@ -61,7 +61,16 @@ UDMXProtocolSettings::UDMXProtocolSettings()
 		{ TEXT("Effects"),			TEXT("Effect, Macro, Effects") },
 		{ TEXT("Frost"),			TEXT("") },
 		{ TEXT("Reset"),			TEXT("FixtureReset, FixtureGlobalReset, GlobalReset") },
-
+		{ TEXT("CTC"),				TEXT("") },
+		{ TEXT("Tint"),				TEXT("") },
+		{ TEXT("Color XF"),			TEXT("") },
+		{ TEXT("HSB_Hue"),			TEXT("") },
+		{ TEXT("HSB_Saturation"),	TEXT("") },
+		{ TEXT("HSB_Brightness"),	TEXT("") },
+		{ TEXT("FanMode"),			TEXT("") },
+		{ TEXT("CIE_X"),			TEXT("") },
+		{ TEXT("CIE_Y"),			TEXT("") },
+		{ TEXT("Prism"),			TEXT("") },
 
 		/* Firework, Fountain related */
 
