DMX - Remove the check when an port is updated but doesn't exist anymore in project settings, instead add an ensure condition

diff --git a/Engine/Plugins/VirtualProduction/DMX/DMXProtocol/Source/DMXProtocol/Private/IO/DMXInputPort.cpp b/Engine/Plugins/VirtualProduction/DMX/DMXProtocol/Source/DMXProtocol/Private/IO/DMXInputPort.cpp
--- a/Engine/Plugins/VirtualProduction/DMX/DMXProtocol/Source/DMXProtocol/Private/IO/DMXInputPort.cpp
+++ b/Engine/Plugins/VirtualProduction/DMX/DMXProtocol/Source/DMXProtocol/Private/IO/DMXInputPort.cpp
@@ -429,7 +429,7 @@ void FDMXInputPort::UpdateFromConfig(FDMXInputPortConfig& InOutInputPortConfig,
 	const bool bConfigIsInProjectSettings = ProtocolSettings->InputPortConfigs.ContainsByPredicate([&InputPortConfig](const FDMXInputPortConfig& Other) {
 		return InputPortConfig.GetPortGuid() == Other.GetPortGuid();
 		});
-	checkf(bConfigIsInProjectSettings, TEXT("Can only use configs with a guid that corresponds to a config in project settings"));
+	ensureAlwaysMsgf(bConfigIsInProjectSettings, TEXT("Can only use configs with a guid that corresponds to a config in project settings"));
 
 	// Find if the port needs update its registration with the protocol
 	const bool bNeedsUpdateRegistration = [this, &InputPortConfig, bForceUpdateRegistrationWithProtocol]()
diff --git a/Engine/Plugins/VirtualProduction/DMX/DMXProtocol/Source/DMXProtocol/Private/IO/DMXOutputPort.cpp b/Engine/Plugins/VirtualProduction/DMX/DMXProtocol/Source/DMXProtocol/Private/IO/DMXOutputPort.cpp
--- a/Engine/Plugins/VirtualProduction/DMX/DMXProtocol/Source/DMXProtocol/Private/IO/DMXOutputPort.cpp
+++ b/Engine/Plugins/VirtualProduction/DMX/DMXProtocol/Source/DMXProtocol/Private/IO/DMXOutputPort.cpp
@@ -556,7 +556,7 @@ void FDMXOutputPort::UpdateFromConfig(FDMXOutputPortConfig& InOutOutputPortConfi
 	const bool bConfigIsInProjectSettings = ProtocolSettings->OutputPortConfigs.ContainsByPredicate([&OutputPortConfig](const FDMXOutputPortConfig& Other) {
 		return OutputPortConfig.GetPortGuid() == Other.GetPortGuid();
 	});
-	checkf(bConfigIsInProjectSettings, TEXT("Can only use configs with a guid that corresponds to a config in project settings"));
+	ensureAlwaysMsgf(bConfigIsInProjectSettings, TEXT("Can only use configs with a guid that corresponds to a config in project settings"));
 
 	// Find if the port needs update its registration with the protocol
 	const bool bNeedsUpdateRegistration = [this, &OutputPortConfig, bForceUpdateRegistrationWithProtocol]()
