DMX - Fix a crash when setting an IP address that doesn't exist on the machine on a DMX port

diff --git a/Engine/Plugins/VirtualProduction/DMX/DMXProtocol/Source/DMXProtocol/Private/IO/DMXInputPort.cpp b/Engine/Plugins/VirtualProduction/DMX/DMXProtocol/Source/DMXProtocol/Private/IO/DMXInputPort.cpp
--- a/Engine/Plugins/VirtualProduction/DMX/DMXProtocol/Source/DMXProtocol/Private/IO/DMXInputPort.cpp
+++ b/Engine/Plugins/VirtualProduction/DMX/DMXProtocol/Source/DMXProtocol/Private/IO/DMXInputPort.cpp
@@ -3,6 +3,7 @@
 #include "IO/DMXInputPort.h"
 
 #include "DMXProtocolSettings.h"
+#include "DMXProtocolUtils.h"
 #include "DMXStats.h"
 #include "Interfaces/IDMXProtocol.h"
 #include "IO/DMXInputPortConfig.h"
@@ -480,13 +481,21 @@ void FDMXInputPort::UpdateFromConfig(FDMXInputPortConfig& InOutInputPortConfig,
 	PriorityStrategy = InputPortConfig.GetPortPriorityStrategy();
 	Priority = InputPortConfig.GetPriority();
 
+	const bool bValidDeviceAddress = FDMXProtocolUtils::GetLocalNetworkInterfaceCardIPs().ContainsByPredicate([this](const TSharedPtr<FString>& IPAddress)
+		{
+			return *IPAddress == DeviceAddress;
+		});
+	if (!bValidDeviceAddress)
+	{
+		UE_LOG(LogDMXProtocol, Warning, TEXT("Cannot register Input Port %s. Device Address '%s' does not exist on the local machine."), *PortName, *DeviceAddress);
+	}
+
 	// Re-register the port if required
-	if (bNeedsUpdateRegistration)
+	if (bNeedsUpdateRegistration && bValidDeviceAddress)
 	{
 		Register();
 	}
 
-
 	OnPortUpdated.Broadcast();
 }
 
diff --git a/Engine/Plugins/VirtualProduction/DMX/DMXProtocol/Source/DMXProtocol/Private/IO/DMXOutputPort.cpp b/Engine/Plugins/VirtualProduction/DMX/DMXProtocol/Source/DMXProtocol/Private/IO/DMXOutputPort.cpp
--- a/Engine/Plugins/VirtualProduction/DMX/DMXProtocol/Source/DMXProtocol/Private/IO/DMXOutputPort.cpp
+++ b/Engine/Plugins/VirtualProduction/DMX/DMXProtocol/Source/DMXProtocol/Private/IO/DMXOutputPort.cpp
@@ -4,6 +4,7 @@
 
 #include "DMXProtocolLog.h"
 #include "DMXProtocolSettings.h"
+#include "DMXProtocolUtils.h"
 #include "Interfaces/IDMXProtocol.h"
 #include "Interfaces/IDMXSender.h"
 #include "IO/DMXOutputPortConfig.h"
@@ -615,8 +616,17 @@ void FDMXOutputPort::UpdateFromConfig(FDMXOutputPortConfig& InOutOutputPortConfi
 
 	CommunicationDeterminator.SetLoopbackToEngine(OutputPortConfig.NeedsLoopbackToEngine());
 
+	const bool bValidDeviceAddress = FDMXProtocolUtils::GetLocalNetworkInterfaceCardIPs().ContainsByPredicate([this](const TSharedPtr<FString>& IPAddress)
+		{
+			return *IPAddress == DeviceAddress;
+		});
+	if (!bValidDeviceAddress)
+	{
+		UE_LOG(LogDMXProtocol, Warning, TEXT("Cannot register Output Port %s. Device Address '%s' does not exist on the local machine."), *PortName, *DeviceAddress);
+	}
+
 	// Re-register the port if required
-	if (bNeedsUpdateRegistration)
+	if (bNeedsUpdateRegistration && bValidDeviceAddress)
 	{
 		Register();
 	}
