DMX - Default input port pritority strategy to 'Highest' to improve support for some DMX consoles (e.g. ETC) out of the box

diff --git a/Engine/Plugins/VirtualProduction/DMX/DMXProtocol/Source/DMXProtocol/Public/IO/DMXInputPortConfig.h b/Engine/Plugins/VirtualProduction/DMX/DMXProtocol/Source/DMXProtocol/Public/IO/DMXInputPortConfig.h
--- a/Engine/Plugins/VirtualProduction/DMX/DMXProtocol/Source/DMXProtocol/Public/IO/DMXInputPortConfig.h
+++ b/Engine/Plugins/VirtualProduction/DMX/DMXProtocol/Source/DMXProtocol/Public/IO/DMXInputPortConfig.h
@@ -135,7 +135,7 @@ protected:
 
 	/** How to deal with the priority value */
 	UPROPERTY(Config, BlueprintReadOnly, EditDefaultsOnly, Category = "Port Config")
-	EDMXPortPriorityStrategy PriorityStrategy = EDMXPortPriorityStrategy::None;
+	EDMXPortPriorityStrategy PriorityStrategy = EDMXPortPriorityStrategy::Highest;
 
 	/** Priority value, can act as a filter or a threshold */
 	UPROPERTY(Config, BlueprintReadOnly, EditDefaultsOnly, Category = "Port Config")
