DMX - Platforms other than Windows now accept 0.0.0.0 as a valid IP for ports

diff --git a/Engine/Plugins/VirtualProduction/DMX/DMXProtocol/Source/DMXProtocol/Private/DMXProtocolUtils.cpp b/Engine/Plugins/VirtualProduction/DMX/DMXProtocol/Source/DMXProtocol/Private/DMXProtocolUtils.cpp
--- a/Engine/Plugins/VirtualProduction/DMX/DMXProtocol/Source/DMXProtocol/Private/DMXProtocolUtils.cpp
+++ b/Engine/Plugins/VirtualProduction/DMX/DMXProtocol/Source/DMXProtocol/Private/DMXProtocolUtils.cpp
@@ -11,11 +11,11 @@
 TArray<TSharedPtr<FString>> FDMXProtocolUtils::GetLocalNetworkInterfaceCardIPs()
 {
 	TArray<TSharedPtr<FString>> LocalNetworkInterfaceCardIPs;
-#if PLATFORM_WINDOWS
-	// Add the default route IP Address, only for windows
+
+	// Add the default route IP Address
 	const FString DefaultRouteLocalAdapterAddress = TEXT("0.0.0.0");
 	LocalNetworkInterfaceCardIPs.Add(MakeShared<FString>(DefaultRouteLocalAdapterAddress));
-#endif 
+
 	// Add the local host IP address
 	const FString LocalHostIpAddress = TEXT("127.0.0.1");
 	LocalNetworkInterfaceCardIPs.Add(MakeShared<FString>(LocalHostIpAddress));
