DMX - Fix a crash when a DMX Output Port is set to sACN Unicast

diff --git a/Engine/Plugins/VirtualProduction/DMX/DMXProtocol/Source/DMXProtocol/Private/IO/DMXOutputPort.cpp b/Engine/Plugins/VirtualProduction/DMX/DMXProtocol/Source/DMXProtocol/Private/IO/DMXOutputPort.cpp
--- a/Engine/Plugins/VirtualProduction/DMX/DMXProtocol/Source/DMXProtocol/Private/IO/DMXOutputPort.cpp
+++ b/Engine/Plugins/VirtualProduction/DMX/DMXProtocol/Source/DMXProtocol/Private/IO/DMXOutputPort.cpp
@@ -616,17 +616,8 @@ void FDMXOutputPort::UpdateFromConfig(FDMXOutputPortConfig& InOutOutputPortConfi
 
 	CommunicationDeterminator.SetLoopbackToEngine(OutputPortConfig.NeedsLoopbackToEngine());
 
-	const bool bValidDeviceAddress = FDMXProtocolUtils::GetLocalNetworkInterfaceCardIPs().ContainsByPredicate([this](const TSharedPtr<FString>& IPAddress)
-		{
-			return *IPAddress == DeviceAddress;
-		});
-	if (!bValidDeviceAddress)
-	{
-		UE_LOG(LogDMXProtocol, Warning, TEXT("Cannot register Output Port %s. Device Address '%s' does not exist on the local machine."), *PortName, *DeviceAddress);
-	}
-
 	// Re-register the port if required
-	if (bNeedsUpdateRegistration && bValidDeviceAddress)
+	if (bNeedsUpdateRegistration)
 	{
 		Register();
 	}
diff --git a/Engine/Plugins/VirtualProduction/DMX/DMXProtocol/Source/DMXProtocolSACN/Private/DMXProtocolSACN.cpp b/Engine/Plugins/VirtualProduction/DMX/DMXProtocol/Source/DMXProtocolSACN/Private/DMXProtocolSACN.cpp
--- a/Engine/Plugins/VirtualProduction/DMX/DMXProtocol/Source/DMXProtocolSACN/Private/DMXProtocolSACN.cpp
+++ b/Engine/Plugins/VirtualProduction/DMX/DMXProtocol/Source/DMXProtocolSACN/Private/DMXProtocolSACN.cpp
@@ -151,18 +151,20 @@ TArray<TSharedPtr<IDMXSender>> FDMXProtocolSACN::RegisterOutputPort(const TShare
 
 	// Try to use an existing receiver or create a new one
 	TArray<TSharedPtr<IDMXSender>> NewSenders;
-	
 	if (CommunicationType == EDMXCommunicationType::Multicast)
 	{
 		TSharedPtr<FDMXProtocolSACNSender> NewSender = FindExistingMulticastSender(NetworkInterfaceAddress);
 		if (!NewSender.IsValid())
 		{
 			NewSender = FDMXProtocolSACNSender::TryCreateMulticastSender(SharedThis(this), NetworkInterfaceAddress);
-			Senders.Add(NewSender);
 		}
 
-		NewSender->AssignOutputPort(OutputPort);
-		NewSenders.Add(NewSender);
+		if (NewSender.IsValid())
+		{
+			Senders.Add(NewSender);
+			NewSender->AssignOutputPort(OutputPort);
+			NewSenders.Add(NewSender);
+		}
 	}
 	else if (CommunicationType == EDMXCommunicationType::Unicast)
 	{
@@ -177,11 +179,14 @@ TArray<TSharedPtr<IDMXSender>> FDMXProtocolSACN::RegisterOutputPort(const TShare
 			if (!NewSender.IsValid())
 			{
 				NewSender = FDMXProtocolSACNSender::TryCreateUnicastSender(SharedThis(this), NetworkInterfaceAddress, UnicastAddress);
-				Senders.Add(NewSender);
 			}
 
-			NewSender->AssignOutputPort(OutputPort);
-			NewSenders.Add(NewSender);
+			if (NewSender.IsValid())
+			{
+				Senders.Add(NewSender);
+				NewSender->AssignOutputPort(OutputPort);
+				NewSenders.Add(NewSender);
+			}
 		}
 	}
 	else
