DMX - Add an option to auto-update IP addresses from ports to an existing IP on the machine

diff --git a/Engine/Plugins/VirtualProduction/DMX/DMXProtocol/Source/DMXProtocol/Private/IO/DMXInputPort.cpp b/Engine/Plugins/VirtualProduction/DMX/DMXProtocol/Source/DMXProtocol/Private/IO/DMXInputPort.cpp
--- a/Engine/Plugins/VirtualProduction/DMX/DMXProtocol/Source/DMXProtocol/Private/IO/DMXInputPort.cpp
+++ b/Engine/Plugins/VirtualProduction/DMX/DMXProtocol/Source/DMXProtocol/Private/IO/DMXInputPort.cpp
@@ -406,6 +406,8 @@ FDMXInputPortConfig FDMXInputPort::MakeInputPortConfig() const
 	FDMXInputPortConfigParams Params;
 	Params.PortName = PortName;
 	Params.ProtocolName = Protocol.IsValid() ? Protocol->GetProtocolName() : NAME_None;
+	Params.bAutoCompleteDeviceAddressEnabled = bAutoCompleteDeviceAddressEnabled;
+	Params.AutoCompleteDeviceAddress = AutoCompleteDeviceAddress;
 	Params.CommunicationType = CommunicationType;
 	Params.DeviceAddress = DeviceAddress;
 	Params.LocalUniverseStart = LocalUniverseStart;
@@ -453,6 +455,8 @@ void FDMXInputPort::UpdateFromConfig(FDMXInputPortConfig& InOutInputPortConfig,
 		FName ProtocolName = Protocol.IsValid() ? Protocol->GetProtocolName() : NAME_None;
 
 		if (ProtocolName == InputPortConfig.GetProtocolName() &&
+			bAutoCompleteDeviceAddressEnabled == bAutoCompleteDeviceAddressEnabled &&
+			AutoCompleteDeviceAddress == AutoCompleteDeviceAddress &&
 			DeviceAddress == InputPortConfig.GetDeviceAddress() &&
 			CommunicationType == InputPortConfig.GetCommunicationType())
 		{
@@ -473,8 +477,10 @@ void FDMXInputPort::UpdateFromConfig(FDMXInputPortConfig& InOutInputPortConfig,
 	// Copy properties from the config into the base class
 	PortGuid = InputPortConfig.GetPortGuid();
 	CommunicationType = InputPortConfig.GetCommunicationType();
-	ExternUniverseStart = InputPortConfig.GetExternUniverseStart();
+	bAutoCompleteDeviceAddressEnabled = InputPortConfig.IsAutoCompleteDeviceAddressEnabled();
+	AutoCompleteDeviceAddress = InputPortConfig.GetAutoCompleteDeviceAddress();
 	DeviceAddress = InputPortConfig.GetDeviceAddress();
+	ExternUniverseStart = InputPortConfig.GetExternUniverseStart();
 	LocalUniverseStart = InputPortConfig.GetLocalUniverseStart();
 	NumUniverses = InputPortConfig.GetNumUniverses();
 	PortName = InputPortConfig.GetPortName();
diff --git a/Engine/Plugins/VirtualProduction/DMX/DMXProtocol/Source/DMXProtocol/Private/IO/DMXInputPortConfig.cpp b/Engine/Plugins/VirtualProduction/DMX/DMXProtocol/Source/DMXProtocol/Private/IO/DMXInputPortConfig.cpp
--- a/Engine/Plugins/VirtualProduction/DMX/DMXProtocol/Source/DMXProtocol/Private/IO/DMXInputPortConfig.cpp
+++ b/Engine/Plugins/VirtualProduction/DMX/DMXProtocol/Source/DMXProtocol/Private/IO/DMXInputPortConfig.cpp
@@ -14,6 +14,8 @@ FDMXInputPortConfigParams::FDMXInputPortConfigParams(const FDMXInputPortConfig&
 	: PortName(InputPortConfig.GetPortName())
 	, ProtocolName(InputPortConfig.GetProtocolName())
 	, CommunicationType(InputPortConfig.GetCommunicationType())
+	, bAutoCompleteDeviceAddressEnabled(InputPortConfig.IsAutoCompleteDeviceAddressEnabled())
+	, AutoCompleteDeviceAddress(InputPortConfig.GetAutoCompleteDeviceAddress())
 	, DeviceAddress(InputPortConfig.GetDeviceAddress())
 	, LocalUniverseStart(InputPortConfig.GetLocalUniverseStart())
 	, NumUniverses(InputPortConfig.GetNumUniverses())
@@ -41,6 +43,8 @@ FDMXInputPortConfig::FDMXInputPortConfig(const FGuid& InPortGuid, const FDMXInpu
 	: PortName(InitializationData.PortName)
 	, ProtocolName(InitializationData.ProtocolName)
 	, CommunicationType(InitializationData.CommunicationType)
+	, bAutoCompleteDeviceAddressEnabled(InitializationData.bAutoCompleteDeviceAddressEnabled)
+	, AutoCompleteDeviceAddress(InitializationData.AutoCompleteDeviceAddress)
 	, DeviceAddress(InitializationData.DeviceAddress)
 	, LocalUniverseStart(InitializationData.LocalUniverseStart)
 	, NumUniverses(InitializationData.NumUniverses)
@@ -124,7 +128,7 @@ void FDMXInputPortConfig::MakeValid()
 
 FString FDMXInputPortConfig::GetDeviceAddress() const
 {
-	// Allow to override the listener ip from commandline
+	// Return the Command Line Device Address if it is set
 	const FString DMXInputPortCommandLine = FString::Printf(TEXT("dmxinputportip=%s:"), *PortName);
 	FString OverrideIP;
 	FParse::Value(FCommandLine::Get(), *DMXInputPortCommandLine, OverrideIP);
@@ -132,5 +136,20 @@ FString FDMXInputPortConfig::GetDeviceAddress() const
 	{
 		return OverrideIP;
 	}
+
+	// Return the Auto Complete Device Address if that is enabled
+	if (bAutoCompleteDeviceAddressEnabled)
+	{
+		const TArray<TSharedPtr<FString>> AvailableNetworkInterfaceCardIPs = FDMXProtocolUtils::GetLocalNetworkInterfaceCardIPs();
+		for (const TSharedPtr<FString>& NetworkInterfaceCardIP : AvailableNetworkInterfaceCardIPs)
+		{
+			if ((*NetworkInterfaceCardIP).MatchesWildcard(AutoCompleteDeviceAddress))
+			{
+				return *NetworkInterfaceCardIP;
+			}
+		}
+	}
+
+	// Return the Device Address
 	return DeviceAddress;
-}
\ No newline at end of file
+}
diff --git a/Engine/Plugins/VirtualProduction/DMX/DMXProtocol/Source/DMXProtocol/Private/IO/DMXOutputPort.cpp b/Engine/Plugins/VirtualProduction/DMX/DMXProtocol/Source/DMXProtocol/Private/IO/DMXOutputPort.cpp
--- a/Engine/Plugins/VirtualProduction/DMX/DMXProtocol/Source/DMXProtocol/Private/IO/DMXOutputPort.cpp
+++ b/Engine/Plugins/VirtualProduction/DMX/DMXProtocol/Source/DMXProtocol/Private/IO/DMXOutputPort.cpp
@@ -531,6 +531,8 @@ FDMXOutputPortConfig FDMXOutputPort::MakeOutputPortConfig() const
 	Params.PortName = PortName;
 	Params.ProtocolName = Protocol.IsValid() ? Protocol->GetProtocolName() : NAME_None;
 	Params.CommunicationType = CommunicationType;
+	Params.bAutoCompleteDeviceAddressEnabled = bAutoCompleteDeviceAddressEnabled;
+	Params.AutoCompleteDeviceAddress = AutoCompleteDeviceAddress;
 	Params.DeviceAddress = DeviceAddress;
 	Params.DestinationAddresses = DestinationAddresses;
 	Params.bLoopbackToEngine = CommunicationDeterminator.IsLoopbackToEngineEnabled();
@@ -575,6 +577,8 @@ void FDMXOutputPort::UpdateFromConfig(FDMXOutputPortConfig& InOutOutputPortConfi
 		FName ProtocolName = Protocol.IsValid() ? Protocol->GetProtocolName() : NAME_None;
 
 		if (ProtocolName == OutputPortConfig.GetProtocolName() &&
+			bAutoCompleteDeviceAddressEnabled == bAutoCompleteDeviceAddressEnabled &&
+			AutoCompleteDeviceAddress == AutoCompleteDeviceAddress &&
 			DeviceAddress == OutputPortConfig.GetDeviceAddress() &&
 			DestinationAddresses == OutputPortConfig.GetDestinationAddresses() &&
 			CommunicationType == OutputPortConfig.GetCommunicationType() &&
@@ -604,11 +608,13 @@ void FDMXOutputPort::UpdateFromConfig(FDMXOutputPortConfig& InOutOutputPortConfi
 	PortGuid = ConfigPortGuid;
 
 	CommunicationType = OutputPortConfig.GetCommunicationType();
-	ExternUniverseStart = OutputPortConfig.GetExternUniverseStart();
+	bAutoCompleteDeviceAddressEnabled = OutputPortConfig.IsAutoCompleteDeviceAddressEnabled();
+	AutoCompleteDeviceAddress = OutputPortConfig.GetAutoCompleteDeviceAddress();
 	DeviceAddress = OutputPortConfig.GetDeviceAddress();
 	DestinationAddresses = OutputPortConfig.GetDestinationAddresses();
 	LocalUniverseStart = OutputPortConfig.GetLocalUniverseStart();
 	NumUniverses = OutputPortConfig.GetNumUniverses();
+	ExternUniverseStart = OutputPortConfig.GetExternUniverseStart();
 	PortName = OutputPortConfig.GetPortName();
 	Priority = OutputPortConfig.GetPriority();
 	DelaySeconds = OutputPortConfig.GetDelay() / OutputPortConfig.GetDelayFrameRate().AsDecimal();
diff --git a/Engine/Plugins/VirtualProduction/DMX/DMXProtocol/Source/DMXProtocol/Private/IO/DMXOutputPortConfig.cpp b/Engine/Plugins/VirtualProduction/DMX/DMXProtocol/Source/DMXProtocol/Private/IO/DMXOutputPortConfig.cpp
--- a/Engine/Plugins/VirtualProduction/DMX/DMXProtocol/Source/DMXProtocol/Private/IO/DMXOutputPortConfig.cpp
+++ b/Engine/Plugins/VirtualProduction/DMX/DMXProtocol/Source/DMXProtocol/Private/IO/DMXOutputPortConfig.cpp
@@ -14,6 +14,8 @@ FDMXOutputPortConfigParams::FDMXOutputPortConfigParams(const FDMXOutputPortConfi
 	: PortName(OutputPortConfig.GetPortName())
 	, ProtocolName(OutputPortConfig.GetProtocolName())
 	, CommunicationType(OutputPortConfig.GetCommunicationType())
+	, bAutoCompleteDeviceAddressEnabled(OutputPortConfig.IsAutoCompleteDeviceAddressEnabled())
+	, AutoCompleteDeviceAddress(OutputPortConfig.GetAutoCompleteDeviceAddress())
 	, DeviceAddress(OutputPortConfig.GetDeviceAddress())
 	, DestinationAddresses(OutputPortConfig.GetDestinationAddresses())
 	, bLoopbackToEngine(OutputPortConfig.NeedsLoopbackToEngine())
@@ -46,6 +48,8 @@ FDMXOutputPortConfig::FDMXOutputPortConfig(const FGuid& InPortGuid, const FDMXOu
 	: PortName(InitializationData.PortName)
 	, ProtocolName(InitializationData.ProtocolName)
 	, CommunicationType(InitializationData.CommunicationType)
+	, bAutoCompleteDeviceAddressEnabled(InitializationData.bAutoCompleteDeviceAddressEnabled)
+	, AutoCompleteDeviceAddress(InitializationData.AutoCompleteDeviceAddress)
 	, DeviceAddress(InitializationData.DeviceAddress)
 	, DestinationAddresses(InitializationData.DestinationAddresses)
 	, bLoopbackToEngine(InitializationData.bLoopbackToEngine)
@@ -142,7 +146,7 @@ void FDMXOutputPortConfig::MakeValid()
 
 FString FDMXOutputPortConfig::GetDeviceAddress() const
 {
-	// Allow to override the source ip from commandline
+	// Return the Command Line Device Address if it is set
 	const FString DMXOutputPortCommandLine = FString::Printf(TEXT("dmxoutputportip=%s:"), *PortName);
 	FString OverrideIP;
 	FParse::Value(FCommandLine::Get(), *DMXOutputPortCommandLine, OverrideIP);
@@ -150,5 +154,20 @@ FString FDMXOutputPortConfig::GetDeviceAddress() const
 	{
 		return OverrideIP;
 	}
+
+	// Return the Auto Complete Device Address if that is enabled
+	if (bAutoCompleteDeviceAddressEnabled)
+	{
+		const TArray<TSharedPtr<FString>> AvailableNetworkInterfaceCardIPs = FDMXProtocolUtils::GetLocalNetworkInterfaceCardIPs();
+		for (const TSharedPtr<FString>& NetworkInterfaceCardIP : AvailableNetworkInterfaceCardIPs)
+		{
+			if ((*NetworkInterfaceCardIP).MatchesWildcard(AutoCompleteDeviceAddress))
+			{
+				return *NetworkInterfaceCardIP;
+			}
+		}
+	}
+
+	// Return the Device Address
 	return DeviceAddress;
-}
\ No newline at end of file
+}
diff --git a/Engine/Plugins/VirtualProduction/DMX/DMXProtocol/Source/DMXProtocol/Public/IO/DMXInputPortConfig.h b/Engine/Plugins/VirtualProduction/DMX/DMXProtocol/Source/DMXProtocol/Public/IO/DMXInputPortConfig.h
--- a/Engine/Plugins/VirtualProduction/DMX/DMXProtocol/Source/DMXProtocol/Public/IO/DMXInputPortConfig.h
+++ b/Engine/Plugins/VirtualProduction/DMX/DMXProtocol/Source/DMXProtocol/Public/IO/DMXInputPortConfig.h
@@ -46,6 +46,8 @@ struct DMXPROTOCOL_API FDMXInputPortConfigParams
 	FString PortName;
 	FName ProtocolName;
 	EDMXCommunicationType CommunicationType;
+	bool bAutoCompleteDeviceAddressEnabled;
+	FString AutoCompleteDeviceAddress;
 	FString DeviceAddress; 
 	int32 LocalUniverseStart;
 	int32 NumUniverses;
@@ -80,6 +82,8 @@ public:
 	FORCEINLINE const FString& GetPortName() const { return PortName; }
 	FORCEINLINE const FName& GetProtocolName() const { return ProtocolName; }
 	FORCEINLINE EDMXCommunicationType GetCommunicationType() const { return CommunicationType; }
+	FORCEINLINE bool IsAutoCompleteDeviceAddressEnabled() const { return bAutoCompleteDeviceAddressEnabled; }
+	FORCEINLINE FString GetAutoCompleteDeviceAddress() const { return AutoCompleteDeviceAddress; }
 	FString GetDeviceAddress() const;
 	FORCEINLINE int32 GetLocalUniverseStart() const { return LocalUniverseStart; }
 	FORCEINLINE int32 GetNumUniverses() const { return NumUniverses; }
@@ -92,6 +96,8 @@ public:
 	static FName GetPortNamePropertyNameChecked() { return GET_MEMBER_NAME_CHECKED(FDMXInputPortConfig, PortName); }
 	static FName GetProtocolNamePropertyNameChecked() { return GET_MEMBER_NAME_CHECKED(FDMXInputPortConfig, ProtocolName); }
 	static FName GetCommunicationTypePropertyNameChecked() { return GET_MEMBER_NAME_CHECKED(FDMXInputPortConfig, CommunicationType); }
+	static FName GetAutoCompleteDeviceAddressEnabledPropertyNameChecked() { return GET_MEMBER_NAME_CHECKED(FDMXInputPortConfig, bAutoCompleteDeviceAddressEnabled); }
+	static FName GetAutoCompleteDeviceAddressPropertyNameChecked() { return GET_MEMBER_NAME_CHECKED(FDMXInputPortConfig, AutoCompleteDeviceAddress); }
 	static FName GetDeviceAddressPropertyNameChecked() { return GET_MEMBER_NAME_CHECKED(FDMXInputPortConfig, DeviceAddress); }
 	static FName GetLocalUniverseStartPropertyNameChecked() { return GET_MEMBER_NAME_CHECKED(FDMXInputPortConfig, LocalUniverseStart); }
 	static FName GetNumUniversesPropertyNameChecked() { return GET_MEMBER_NAME_CHECKED(FDMXInputPortConfig, NumUniverses); }
@@ -114,6 +120,23 @@ protected:
 	UPROPERTY(Config, BlueprintReadOnly, EditDefaultsOnly, Category = "Port Config")
 	EDMXCommunicationType CommunicationType = EDMXCommunicationType::InternalOnly;
 
+	/** Enables 'Auto Complete Device Address', hidden via customization - EditConditionInlineToggle doesn't support Config */
+	UPROPERTY(Config, BlueprintReadOnly, EditDefaultsOnly, Category = "Port Config")
+	bool bAutoCompleteDeviceAddressEnabled = false;
+
+	/**
+	 * Searches available Network Interface Card IP Addresses and uses the first match as the 'Network Interface Card IP Address' (both in Editor and Game).
+	 *
+	 * Supports wildcards, examples:
+	 * '192'
+	 * '192.*'
+	 * '192.168.?.*'.
+	 *
+	 * If empty or '*' the first best available IP will be selected (not recommended)
+	 */
+	UPROPERTY(Config, BlueprintReadOnly, EditDefaultsOnly, Category = "Port Config", Meta = (DisplayName = "Auto Complete Network Interface Card IP Address"))
+	FString AutoCompleteDeviceAddress = TEXT("192.*");
+
 	/** The Network Interface Card's IP Adress, over which DMX is received */
 	UPROPERTY(Config, BlueprintReadOnly, EditDefaultsOnly, Category = "Port Config", Meta = (DisplayName = "Network Interface Card IP Address"))
 	FString DeviceAddress = TEXT("127.0.0.1");
diff --git a/Engine/Plugins/VirtualProduction/DMX/DMXProtocol/Source/DMXProtocol/Public/IO/DMXOutputPortConfig.h b/Engine/Plugins/VirtualProduction/DMX/DMXProtocol/Source/DMXProtocol/Public/IO/DMXOutputPortConfig.h
--- a/Engine/Plugins/VirtualProduction/DMX/DMXProtocol/Source/DMXProtocol/Public/IO/DMXOutputPortConfig.h
+++ b/Engine/Plugins/VirtualProduction/DMX/DMXProtocol/Source/DMXProtocol/Public/IO/DMXOutputPortConfig.h
@@ -46,6 +46,8 @@ struct DMXPROTOCOL_API FDMXOutputPortConfigParams
 	FString PortName;
 	FName ProtocolName;
 	EDMXCommunicationType CommunicationType;
+	bool bAutoCompleteDeviceAddressEnabled;
+	FString AutoCompleteDeviceAddress;
 	FString DeviceAddress;
 	TArray<FDMXOutputPortDestinationAddress> DestinationAddresses;
 	bool bLoopbackToEngine;
@@ -86,6 +88,8 @@ public:
 	FORCEINLINE const FString& GetPortName() const { return PortName; }
 	FORCEINLINE const FName& GetProtocolName() const { return ProtocolName; }
 	FORCEINLINE EDMXCommunicationType GetCommunicationType() const { return CommunicationType; }
+	FORCEINLINE bool IsAutoCompleteDeviceAddressEnabled() const { return bAutoCompleteDeviceAddressEnabled; }
+	FORCEINLINE FString GetAutoCompleteDeviceAddress() const { return AutoCompleteDeviceAddress; }
 	FString GetDeviceAddress() const;
 	FORCEINLINE const TArray<FDMXOutputPortDestinationAddress>& GetDestinationAddresses() const { return DestinationAddresses; }
 	FORCEINLINE bool NeedsLoopbackToEngine() const { return bLoopbackToEngine; }
@@ -104,6 +108,8 @@ public:
 	static FName GetPortNamePropertyNameChecked() { return GET_MEMBER_NAME_CHECKED(FDMXOutputPortConfig, PortName); }
 	static FName GetProtocolNamePropertyNameChecked() { return GET_MEMBER_NAME_CHECKED(FDMXOutputPortConfig, ProtocolName); }
 	static FName GetCommunicationTypePropertyNameChecked() { return GET_MEMBER_NAME_CHECKED(FDMXOutputPortConfig, CommunicationType); }
+	static FName GetAutoCompleteDeviceAddressEnabledPropertyNameChecked() { return GET_MEMBER_NAME_CHECKED(FDMXOutputPortConfig, bAutoCompleteDeviceAddressEnabled); }
+	static FName GetAutoCompleteDeviceAddressPropertyNameChecked() { return GET_MEMBER_NAME_CHECKED(FDMXOutputPortConfig, AutoCompleteDeviceAddress); }
 	static FName GetDeviceAddressPropertyNameChecked() { return GET_MEMBER_NAME_CHECKED(FDMXOutputPortConfig, DeviceAddress); }
 	static FName GetDestinationAddressesPropertyNameChecked() { return GET_MEMBER_NAME_CHECKED(FDMXOutputPortConfig, DestinationAddresses); }
 	static FName GetLocalUniverseStartPropertyNameChecked() { return GET_MEMBER_NAME_CHECKED(FDMXOutputPortConfig, LocalUniverseStart); }
@@ -131,6 +137,23 @@ protected:
 	UPROPERTY(Config, BlueprintReadWrite, EditDefaultsOnly, Category = "Port Config")
 	EDMXCommunicationType CommunicationType = EDMXCommunicationType::InternalOnly;
 
+	/** Enables 'Auto Complete Device Address', hidden via customization - EditConditionInlineToggle doesn't support Config */
+	UPROPERTY(Config, BlueprintReadOnly, EditDefaultsOnly, Category = "Port Config")
+	bool bAutoCompleteDeviceAddressEnabled = false;
+
+	/**
+	 * Searches available Network Interface Card IP Addresses and uses the first match as the 'Network Interface Card IP Address' (both in Editor and Game).
+	 *
+	 * Supports wildcards, examples:
+	 * '192'
+	 * '192.*'
+	 * '192.168.?.*'..
+	 *
+	 * If empty or '*' the first best available IP will be selected (not recommended)
+	 */
+	UPROPERTY(Config, BlueprintReadOnly, EditDefaultsOnly, Category = "Port Config", Meta = (DisplayName = "Auto Complete Network Interface Card IP Address"))
+	FString AutoCompleteDeviceAddress = TEXT("192.*");
+
 	/** The IP address of the network interface card over which outbound DMX is sent */
 	UPROPERTY(Config, BlueprintReadOnly, EditDefaultsOnly, Category = "Port Config", Meta = (DisplayName = "Network Interface Card IP Address"))
 	FString DeviceAddress = TEXT("127.0.0.1");
diff --git a/Engine/Plugins/VirtualProduction/DMX/DMXProtocol/Source/DMXProtocol/Public/IO/DMXPort.h b/Engine/Plugins/VirtualProduction/DMX/DMXProtocol/Source/DMXProtocol/Public/IO/DMXPort.h
--- a/Engine/Plugins/VirtualProduction/DMX/DMXProtocol/Source/DMXProtocol/Public/IO/DMXPort.h
+++ b/Engine/Plugins/VirtualProduction/DMX/DMXProtocol/Source/DMXProtocol/Public/IO/DMXPort.h
@@ -79,11 +79,15 @@ public:
 	int32 ConvertLocalToExternUniverseID(int32 LocalUniverseID) const;
 
 	FORCEINLINE const FString& GetPortName() const { return PortName; }
-	
-	FORCEINLINE const FString& GetDeviceAddress() const { return DeviceAddress; }
 
 	FORCEINLINE const IDMXProtocolPtr& GetProtocol() const { return Protocol; }
 
+	FORCEINLINE bool IsAutoCompleteDeviceAddressEnabled() const { return bAutoCompleteDeviceAddressEnabled; }
+
+	FORCEINLINE const FString& GetAutoCompleteDeviceAddress() const { return AutoCompleteDeviceAddress; }
+
+	FORCEINLINE const FString& GetDeviceAddress() const { return DeviceAddress; }
+
 	FORCEINLINE EDMXCommunicationType GetCommunicationType() const { return CommunicationType; }
 
 	FORCEINLINE int32 GetLocalUniverseStart() const { return LocalUniverseStart; }
@@ -113,6 +117,12 @@ protected:
 	/** The communication type of this port */
 	EDMXCommunicationType CommunicationType;
 
+	/** If true, instead of using the Device Address, the port will auto-complete the IP from the available Network Interface Card Addresses */
+	bool bAutoCompleteDeviceAddressEnabled = false;
+
+	/** The auto-complete Device Address, supports wildcards */
+	FString AutoCompleteDeviceAddress;
+
 	/** The address of the device that handles communication, e.g. the network interface for art-net */
 	FString DeviceAddress;
 
diff --git a/Engine/Plugins/VirtualProduction/DMX/DMXProtocol/Source/DMXProtocolEditor/Private/DetailsCustomizations/DMXInputPortConfigCustomization.cpp b/Engine/Plugins/VirtualProduction/DMX/DMXProtocol/Source/DMXProtocolEditor/Private/DetailsCustomizations/DMXInputPortConfigCustomization.cpp
--- a/Engine/Plugins/VirtualProduction/DMX/DMXProtocol/Source/DMXProtocolEditor/Private/DetailsCustomizations/DMXInputPortConfigCustomization.cpp
+++ b/Engine/Plugins/VirtualProduction/DMX/DMXProtocol/Source/DMXProtocolEditor/Private/DetailsCustomizations/DMXInputPortConfigCustomization.cpp
@@ -5,6 +5,7 @@
 #include "DMXProtocolModule.h"
 #include "DMXProtocolSettings.h"
 #include "DMXProtocolTypes.h"
+#include "DMXProtocolUtils.h"
 #include "Interfaces/IDMXProtocol.h"
 #include "IO/DMXInputPort.h"
 #include "IO/DMXInputPortConfig.h"
@@ -57,9 +58,14 @@ void FDMXInputPortConfigCustomization::CustomizeChildren(TSharedRef<IPropertyHan
 	// Cache customized properties
 	ProtocolNameHandle = PropertyHandles.FindChecked(FDMXInputPortConfig::GetProtocolNamePropertyNameChecked());
 	CommunicationTypeHandle = PropertyHandles.FindChecked(FDMXInputPortConfig::GetCommunicationTypePropertyNameChecked());
+	AutoCompleteDeviceAddressEnabledHandle = PropertyHandles.FindChecked(FDMXInputPortConfig::GetAutoCompleteDeviceAddressEnabledPropertyNameChecked());
+	AutoCompleteDeviceAddressHandle = PropertyHandles.FindChecked(FDMXInputPortConfig::GetAutoCompleteDeviceAddressPropertyNameChecked());
 	DeviceAddressHandle = PropertyHandles.FindChecked(FDMXInputPortConfig::GetDeviceAddressPropertyNameChecked());
 	PortGuidHandle = PropertyHandles.FindChecked(FDMXInputPortConfig::GetPortGuidPropertyNameChecked());
 
+	// Hande property changes
+	AutoCompleteDeviceAddressHandle->SetOnPropertyValueChanged(FSimpleDelegate::CreateSP(this, &FDMXInputPortConfigCustomization::UpdateAutoCompleteDeviceAddressTextBox));
+
 	// Ports always need a valid Guid (cannot be blueprinted)
 	if (!GetPortGuid().IsValid())
 	{
@@ -81,6 +87,12 @@ void FDMXInputPortConfigCustomization::CustomizeChildren(TSharedRef<IPropertyHan
 		{
 			continue;
 		}
+
+		// Don't add the bAutoCompeteDeviceAddressEnabled property
+		if (Iter.Value() == AutoCompleteDeviceAddressEnabledHandle)
+		{
+			continue;
+		}
 		
 		// Add the property
 		IDetailPropertyRow& PropertyRow = ChildBuilder.AddProperty(Iter.Value().ToSharedRef());
@@ -93,29 +105,34 @@ void FDMXInputPortConfigCustomization::CustomizeChildren(TSharedRef<IPropertyHan
 		{
 			GenerateCommunicationTypeRow(PropertyRow);
 		}
+		else if (Iter.Value() == AutoCompleteDeviceAddressHandle)
+		{
+			GenerateAutoCompleteDeviceAddressRow(PropertyRow);
+		}
 		else if (Iter.Value() == DeviceAddressHandle)
 		{
-			GenerateIPAddressRow(PropertyRow);
+			GenerateDeviceAddressRow(PropertyRow);
 		}
 		else if (Iter.Key() == FDMXInputPortConfig::GetPriorityStrategyPropertyNameChecked())
 		{
-			TAttribute<EVisibility> PriorityStrategyVisibilityAttribute = TAttribute<EVisibility>::Create(TAttribute<EVisibility>::FGetter::CreateLambda([this]() 
+
+			PropertyRow.Visibility(TAttribute<EVisibility>::CreateLambda([this]
 				{
 					const bool bPriorityStrategyIsVisible = GetProtocol().IsValid() && GetProtocol()->SupportsPrioritySettings();
-					return bPriorityStrategyIsVisible ? EVisibility::Visible : EVisibility::Collapsed;
+					return bPriorityStrategyIsVisible ? 
+						EVisibility::Visible : 
+						EVisibility::Collapsed;
 				}));
-
-			PropertyRow.Visibility(PriorityStrategyVisibilityAttribute);
 		}
 		else if (Iter.Key() == FDMXInputPortConfig::GetPriorityPropertyNameChecked())
 		{
-			TAttribute<EVisibility> PriorityVisibilityAttribute = TAttribute<EVisibility>::Create(TAttribute<EVisibility>::FGetter::CreateLambda([this]() 
+			PropertyRow.Visibility(TAttribute<EVisibility>::CreateLambda([this]
 				{
 					const bool bPriorityIsVisible = GetProtocol().IsValid() && GetProtocol()->SupportsPrioritySettings();
-					return bPriorityIsVisible ? EVisibility::Visible : EVisibility::Collapsed;
+					return bPriorityIsVisible ? 
+						EVisibility::Visible : 
+						EVisibility::Collapsed;
 				}));
-
-			PropertyRow.Visibility(PriorityVisibilityAttribute);
 		}
 	}
 }
@@ -123,7 +140,6 @@ void FDMXInputPortConfigCustomization::CustomizeChildren(TSharedRef<IPropertyHan
 void FDMXInputPortConfigCustomization::GenerateProtocolNameRow(IDetailPropertyRow& PropertyRow)
 {
 	// Customizate the Protocol Name property to draw a combo box with all protocol names
-
 	TSharedPtr<SWidget> NameWidget;
 	TSharedPtr<SWidget> ValueWidget;
 	FDetailWidgetRow Row;
@@ -156,9 +172,14 @@ void FDMXInputPortConfigCustomization::GenerateProtocolNameRow(IDetailPropertyRo
 void FDMXInputPortConfigCustomization::GenerateCommunicationTypeRow(IDetailPropertyRow& PropertyRow)
 {
 	// Customizate the Communication Type property depending on the port's supported communication modes
-
-	TAttribute<EVisibility> CommunicationTypeVisibilityAttribute = TAttribute<EVisibility>::Create(TAttribute<EVisibility>::FGetter::CreateSP(this, &FDMXInputPortConfigCustomization::GetCommunicationTypeVisibility));
-	PropertyRow.Visibility(CommunicationTypeVisibilityAttribute);
+	const TAttribute<EVisibility> CommunicationTypeVisibilityAttribute = TAttribute<EVisibility>::CreateLambda([this]
+		{
+			if (CommunicationTypeComboBox.IsValid() && GetProtocol().IsValid())
+			{
+				return CommunicationTypeComboBox->GetVisibility();
+			}
+			return EVisibility::Collapsed;
+		});
 
 	TSharedPtr<SWidget> NameWidget;
 	TSharedPtr<SWidget> ValueWidget;
@@ -180,16 +201,107 @@ void FDMXInputPortConfigCustomization::GenerateCommunicationTypeRow(IDetailPrope
 		];
 }
 
-void FDMXInputPortConfigCustomization::GenerateIPAddressRow(IDetailPropertyRow& PropertyRow)
+void FDMXInputPortConfigCustomization::GenerateAutoCompleteDeviceAddressRow(IDetailPropertyRow& PropertyRow)
+{
+	// Mimic a Inline Edit Condition Toggle - Doing it directly via Meta Specifiers is not possible. -
+	// The Inline Edit Condition Toggle property wouldn't get properly serialized when being a Config property.
+	const TSharedRef<SWidget> AutoCompleteDeviceAddressEnableldPropertyValueWidget = AutoCompleteDeviceAddressEnabledHandle->CreatePropertyValueWidget();
+
+	const TSharedRef<SWidget> AutoCompleteDeviceAddressPropertyNameWidget = AutoCompleteDeviceAddressHandle->CreatePropertyNameWidget();
+	const TSharedRef<SWidget> AutoCompleteDeviceAddressPropertyValueWidget = AutoCompleteDeviceAddressHandle->CreatePropertyValueWidget();
+
+	AutoCompleteDeviceAddressPropertyNameWidget->SetEnabled(TAttribute<bool>::CreateLambda([this]()
+		{
+			return IsAutoCompleteDeviceAddressEnabled();
+		}));
+	AutoCompleteDeviceAddressPropertyValueWidget->SetEnabled(TAttribute<bool>::CreateLambda([this]()
+		{
+			return IsAutoCompleteDeviceAddressEnabled();
+		}));
+
+	PropertyRow.CustomWidget()
+		.NameContent()
+		[
+			SNew(SHorizontalBox)
+	
+			// Edit Inline Condition Toggle Checkbox
+			+ SHorizontalBox::Slot()
+			.HAlign(HAlign_Left)
+			.VAlign(VAlign_Fill)
+			.AutoWidth()
+			[
+				AutoCompleteDeviceAddressEnableldPropertyValueWidget
+			]
+
+			// Property Name
+			+ SHorizontalBox::Slot()
+			.Padding(8.f, 0.f, 0.f, 0.f)
+			.HAlign(HAlign_Left)
+			.VAlign(VAlign_Fill)
+			.AutoWidth()
+			[
+				AutoCompleteDeviceAddressPropertyNameWidget
+			]
+		]
+		.ValueContent()
+		[
+			// Property Value
+			AutoCompleteDeviceAddressPropertyValueWidget
+		];
+}
+
+void FDMXInputPortConfigCustomization::GenerateDeviceAddressRow(IDetailPropertyRow& PropertyRow)
 {
 	// Customizate the IPAddress property to show a combo box with available IP addresses
-	FString InitialValue = GetIPAddress();
+	const FString InitialValue = GetIPAddress();
 
 	TSharedPtr<SWidget> NameWidget;
 	TSharedPtr<SWidget> ValueWidget;
 	FDetailWidgetRow Row;
 	PropertyRow.GetDefaultWidgets(NameWidget, ValueWidget, Row);
 
+	// Override the value widget so it either shows 
+	// - The IP Address Edit widget when auto complete is disabled
+	// - The auto completed IP Address when auto complete is enabled
+	ValueWidget =
+		SNew(SHorizontalBox)
+		
+		+ SHorizontalBox::Slot()
+		.HAlign(HAlign_Fill)
+		.VAlign(VAlign_Fill)
+		[
+			SAssignNew(IPAddressEditWidget, SDMXIPAddressEditWidget)
+			.Visibility_Lambda([this]()
+				{
+					return IsAutoCompleteDeviceAddressEnabled() ?
+						EVisibility::Collapsed :
+						EVisibility::Visible;
+				})
+			.InitialValue(InitialValue)
+			.bShowLocalNICComboBox(true)
+			.OnIPAddressSelected(this, &FDMXInputPortConfigCustomization::OnDeviceAddressSelected)
+		]
+
+		+ SHorizontalBox::Slot()
+		.HAlign(HAlign_Fill)
+		.VAlign(VAlign_Fill)
+		[
+			SAssignNew(AutoCompletedDeviceAddressTextBlock, STextBlock)
+			.Visibility_Lambda([this]()
+				{
+					return IsAutoCompleteDeviceAddressEnabled() ?
+						EVisibility::Visible :
+						EVisibility::Collapsed;
+				})
+			.Font(FEditorStyle::GetFontStyle(TEXT("PropertyWindow.NormalFont")))
+		];
+	UpdateAutoCompleteDeviceAddressTextBox();
+
+	PropertyRow.IsEnabled(TAttribute<bool>::CreateLambda([this]
+		{
+			return !IsAutoCompleteDeviceAddressEnabled();
+		}));
+
 	PropertyRow.CustomWidget()
 		.NameContent()
 		[
@@ -197,21 +309,60 @@ void FDMXInputPortConfigCustomization::GenerateIPAddressRow(IDetailPropertyRow&
 		]
 		.ValueContent()
 		[
-			SAssignNew(IPAddressEditWidget, SDMXIPAddressEditWidget)
-			.InitialValue(InitialValue)
-			.bShowLocalNICComboBox(true)
-			.OnIPAddressSelected(this, &FDMXInputPortConfigCustomization::OnIPAddressSelected)
+			ValueWidget.ToSharedRef()
 		];
-}	
+}
 
-EVisibility FDMXInputPortConfigCustomization::GetCommunicationTypeVisibility() const
+void FDMXInputPortConfigCustomization::UpdateAutoCompleteDeviceAddressTextBox()
 {
-	if (CommunicationTypeComboBox.IsValid() && GetProtocol().IsValid())
+	// Set text on the AutoCompleteDeviceAddress Text Box depending on what the resulting IP is
+	if (AutoCompletedDeviceAddressTextBlock.IsValid())
 	{
-		return CommunicationTypeComboBox->GetVisibility();
-	}
+		FString AutoCompleteDeviceAddress;
+		if (AutoCompleteDeviceAddressHandle->GetValue(AutoCompleteDeviceAddress) == FPropertyAccess::Success)
+		{
+			TArray<FString> SearchStrings;
+			constexpr bool bCullEmpty = true;
+			AutoCompleteDeviceAddress.ParseIntoArray(SearchStrings, TEXT("."), bCullEmpty);
 
-	return EVisibility::Collapsed;
+			const TArray<TSharedPtr<FString>> LocalNetworkInterfaceCardIPs = FDMXProtocolUtils::GetLocalNetworkInterfaceCardIPs();
+
+			FString NewIPAddress;
+			for (const TSharedPtr<FString>& NetworkInterfaceCardIP : LocalNetworkInterfaceCardIPs)
+			{
+				TArray<FString> NetworkInterfaceCardIPSubstrings;
+				(*NetworkInterfaceCardIP).ParseIntoArray(NetworkInterfaceCardIPSubstrings, TEXT("."));
+				
+				// Try to match each search string
+				bool bMatchesSearchStrings = true;
+				for (int32 SubstringIndex = 0; SubstringIndex < NetworkInterfaceCardIPSubstrings.Num(); SubstringIndex++)
+				{
+					if (SearchStrings.IsValidIndex(SubstringIndex))
+					{
+						if (!NetworkInterfaceCardIPSubstrings[SubstringIndex].MatchesWildcard(SearchStrings[SubstringIndex]))
+						{
+							bMatchesSearchStrings = false;
+							break;
+						}
+					}
+				}
+
+				if (bMatchesSearchStrings)
+				{
+					// Set the text and return if a match was found
+					const FText IPAddressText = FText::FromString(*NetworkInterfaceCardIP);
+					AutoCompletedDeviceAddressTextBlock->SetText(FText::Format(LOCTEXT("AutoCompletedIPAddressText", "{0} (auto-completed)"), IPAddressText));
+					return;
+				}
+			}
+
+			// Fall back to the manually set one
+			FString DeviceAddress;
+			DeviceAddressHandle->GetValue(DeviceAddress);
+			const FText NoMatchForIPAddressText = FText::Format(LOCTEXT("NoRegexMatchForDeviceAddress", "No match. Using '{0}'"), FText::FromString(DeviceAddress));
+			AutoCompletedDeviceAddressTextBlock->SetText(NoMatchForIPAddressText);
+		}
+	}
 }
 
 void FDMXInputPortConfigCustomization::OnProtocolNameSelected()
@@ -240,7 +391,7 @@ void FDMXInputPortConfigCustomization::OnCommunicationTypeSelected()
 	PropertyUtilities->ForceRefresh();
 }
 
-void FDMXInputPortConfigCustomization::OnIPAddressSelected()
+void FDMXInputPortConfigCustomization::OnDeviceAddressSelected()
 {
 	FString SelectedIP = IPAddressEditWidget->GetSelectedIPAddress();
 
@@ -297,33 +448,16 @@ EDMXCommunicationType FDMXInputPortConfigCustomization::GetCommunicationType() c
 	return static_cast<EDMXCommunicationType>(CommunicationType);
 }
 
-FString FDMXInputPortConfigCustomization::GetIPAddress() const
+bool FDMXInputPortConfigCustomization::IsAutoCompleteDeviceAddressEnabled() const
 {
-	FGuid PortGuid = GetPortGuid();
-	if (PortGuid.IsValid())
-	{
-		const UDMXProtocolSettings* ProtocolSettings = GetDefault<UDMXProtocolSettings>();
-		checkf(ProtocolSettings, TEXT("Unexpected protocol settings not available when its details are customized"));
+	bool bAutoCompleteDeviceAddressEnabled;
+	AutoCompleteDeviceAddressEnabledHandle->GetValue(bAutoCompleteDeviceAddressEnabled);
 
-		const FDMXInputPortConfig* InputPortConfigPtr = ProtocolSettings->InputPortConfigs.FindByPredicate([&PortGuid](const FDMXInputPortConfig& InputPortConfig)
-			{
-				return InputPortConfig.GetPortGuid() == PortGuid;
-			});
-		if (InputPortConfigPtr)
-		{
-			return InputPortConfigPtr->GetDeviceAddress();
-		}
-
-		const FDMXOutputPortConfig* OutputPortConfigPtr = ProtocolSettings->OutputPortConfigs.FindByPredicate([&PortGuid](const FDMXOutputPortConfig& OutputPortConfig)
-			{
-				return OutputPortConfig.GetPortGuid() == PortGuid;
-			});
-		if (OutputPortConfigPtr)
-		{
-			return OutputPortConfigPtr->GetDeviceAddress();
-		}
-	}
+	return bAutoCompleteDeviceAddressEnabled;
+}
 
+FString FDMXInputPortConfigCustomization::GetIPAddress() const
+{
 	FString IPAddress;
 	ensure(DeviceAddressHandle->GetValue(IPAddress) == FPropertyAccess::Success);
 
diff --git a/Engine/Plugins/VirtualProduction/DMX/DMXProtocol/Source/DMXProtocolEditor/Private/DetailsCustomizations/DMXInputPortConfigCustomization.h b/Engine/Plugins/VirtualProduction/DMX/DMXProtocol/Source/DMXProtocolEditor/Private/DetailsCustomizations/DMXInputPortConfigCustomization.h
--- a/Engine/Plugins/VirtualProduction/DMX/DMXProtocol/Source/DMXProtocolEditor/Private/DetailsCustomizations/DMXInputPortConfigCustomization.h
+++ b/Engine/Plugins/VirtualProduction/DMX/DMXProtocol/Source/DMXProtocolEditor/Private/DetailsCustomizations/DMXInputPortConfigCustomization.h
@@ -38,21 +38,24 @@ private:
 	/** Generates the customized Communication Type row */
 	void GenerateCommunicationTypeRow(IDetailPropertyRow& PropertyRow);
 
-	/** Generates the customized IP Address row */
-	void GenerateIPAddressRow(IDetailPropertyRow& PropertyRow);
+	/** Generates the customized Auto Complete Device Address row */
+	void GenerateAutoCompleteDeviceAddressRow(IDetailPropertyRow& PropertyRow);
 
-	/** Returns the visibility of the Communication Type Property */
-	EVisibility GetCommunicationTypeVisibility() const;
+	/** Generates the customized Device Address row */
+	void GenerateDeviceAddressRow(IDetailPropertyRow& PropertyRow);
+
+	/** Updates the Auto Complete Device Address Text Box */
+	void UpdateAutoCompleteDeviceAddressTextBox();
 
 	/** Called when a Protocol Name was selected */
 	void OnProtocolNameSelected();
 
-	/** Called when a local IP Address was selected */
-	void OnIPAddressSelected();
-
 	/** Called when a Communication Type was selected */
 	void OnCommunicationTypeSelected();
 
+	/** Called when a local IP Address was selected */
+	void OnDeviceAddressSelected();
+
 	/** Gets the protocol, checks if it is valid */
 	IDMXProtocolPtr GetProtocol() const;
 
@@ -65,12 +68,21 @@ private:
 	/** Gets the communication type */
 	EDMXCommunicationType GetCommunicationType() const;
 
+	/** Returns the value of the AutoCompleteDeviceAddressEnabled property */
+	bool IsAutoCompleteDeviceAddressEnabled() const;
+
 	/** Gets the IP Address */
 	FString GetIPAddress() const;
 
 	/** Property handle to the ProtocolName property */
 	TSharedPtr<IPropertyHandle> ProtocolNameHandle;
 
+	/** Property handle to the bAutoCompleteDeviceAddressEnabled property */
+	TSharedPtr<IPropertyHandle> AutoCompleteDeviceAddressEnabledHandle;
+
+	/** Property handle to the AutoCompleteDeviceAddress property */
+	TSharedPtr<IPropertyHandle> AutoCompleteDeviceAddressHandle;
+
 	/** Property handle to the IPAddress property */
 	TSharedPtr<IPropertyHandle> DeviceAddressHandle;
 
@@ -80,14 +92,18 @@ private:
 	/** Property handle to the PortGuid property */
 	TSharedPtr<IPropertyHandle> PortGuidHandle;
 
+	/** ComboBox that exposes a selection of communication types to the user */
+	TSharedPtr<SDMXCommunicationTypeComboBox> CommunicationTypeComboBox;
+
 	/** ComboBox to select a protocol name */
 	TSharedPtr<SDMXProtocolNameComboBox> ProtocolNameComboBox;
 
-	/** ComboBox that displays local ip addresses */
+	/** ComboBox that displays the device address when it is editable */
 	TSharedPtr<SDMXIPAddressEditWidget> IPAddressEditWidget;
 
-	/** ComboBox that exposes a selection of communication types to the user */
-	TSharedPtr<SDMXCommunicationTypeComboBox> CommunicationTypeComboBox;
+	/** Text Block that displays the device address when it is auto completed */
+	TSharedPtr<STextBlock> AutoCompletedDeviceAddressTextBlock;
+
 
 	/** Property utilities for this customization */
 	TSharedPtr<IPropertyUtilities> PropertyUtilities;
diff --git a/Engine/Plugins/VirtualProduction/DMX/DMXProtocol/Source/DMXProtocolEditor/Private/DetailsCustomizations/DMXOutputPortConfigCustomization.cpp b/Engine/Plugins/VirtualProduction/DMX/DMXProtocol/Source/DMXProtocolEditor/Private/DetailsCustomizations/DMXOutputPortConfigCustomization.cpp
--- a/Engine/Plugins/VirtualProduction/DMX/DMXProtocol/Source/DMXProtocolEditor/Private/DetailsCustomizations/DMXOutputPortConfigCustomization.cpp
+++ b/Engine/Plugins/VirtualProduction/DMX/DMXProtocol/Source/DMXProtocolEditor/Private/DetailsCustomizations/DMXOutputPortConfigCustomization.cpp
@@ -5,6 +5,7 @@
 #include "DMXProtocolModule.h"
 #include "DMXProtocolSettings.h"
 #include "DMXProtocolTypes.h"
+#include "DMXProtocolUtils.h"
 #include "Interfaces/IDMXProtocol.h"
 #include "IO/DMXOutputPort.h"
 #include "IO/DMXOutputPortConfig.h"
@@ -58,11 +59,16 @@ void FDMXOutputPortConfigCustomization::CustomizeChildren(TSharedRef<IPropertyHa
 	// Cache customized properties
 	ProtocolNameHandle = PropertyHandles.FindChecked(FDMXOutputPortConfig::GetProtocolNamePropertyNameChecked());
 	CommunicationTypeHandle = PropertyHandles.FindChecked(FDMXOutputPortConfig::GetCommunicationTypePropertyNameChecked());
+	AutoCompleteDeviceAddressEnabledHandle = PropertyHandles.FindChecked(FDMXOutputPortConfig::GetAutoCompleteDeviceAddressEnabledPropertyNameChecked());
+	AutoCompleteDeviceAddressHandle = PropertyHandles.FindChecked(FDMXOutputPortConfig::GetAutoCompleteDeviceAddressPropertyNameChecked());
 	DeviceAddressHandle = PropertyHandles.FindChecked(FDMXOutputPortConfig::GetDeviceAddressPropertyNameChecked());
 	DelayHandle = PropertyHandles.FindChecked(FDMXOutputPortConfig::GetDelayPropertyNameChecked());
 	DelayFrameRateHandle = PropertyHandles.FindChecked(FDMXOutputPortConfig::GetDelayFrameRatePropertyNameChecked());
 	PortGuidHandle = PropertyHandles.FindChecked(FDMXOutputPortConfig::GetPortGuidPropertyNameChecked());
 
+	// Hande property changes
+	AutoCompleteDeviceAddressHandle->SetOnPropertyValueChanged(FSimpleDelegate::CreateSP(this, &FDMXOutputPortConfigCustomization::UpdateAutoCompleteDeviceAddressTextBox));
+
 	// Ports always need a valid Guid (cannot be blueprinted)
 	if (!GetPortGuid().IsValid())
 	{
@@ -85,7 +91,13 @@ void FDMXOutputPortConfigCustomization::CustomizeChildren(TSharedRef<IPropertyHa
 		{
 			continue;
 		}
-		
+
+		// Don't add the bAutoCompeteDeviceAddressEnabled property
+		if (Iter.Value() == AutoCompleteDeviceAddressEnabledHandle)
+		{
+			continue;
+		}
+
 		// Add the property
 		IDetailPropertyRow& PropertyRow = ChildBuilder.AddProperty(Iter.Value().ToSharedRef());
 
@@ -97,9 +109,13 @@ void FDMXOutputPortConfigCustomization::CustomizeChildren(TSharedRef<IPropertyHa
 		{
 			GenerateCommunicationTypeRow(PropertyRow);
 		}
+		else if (Iter.Value() == AutoCompleteDeviceAddressHandle)
+		{
+			GenerateAutoCompleteDeviceAddressRow(PropertyRow);
+		}
 		else if (Iter.Value() == DeviceAddressHandle)
 		{
-			GenerateIPAddressRow(PropertyRow);
+			GenerateDeviceAddressRow(PropertyRow);
 		}
 		else if (Iter.Key() == FDMXOutputPortConfig::GetDestinationAddressesPropertyNameChecked())
 		{
@@ -165,17 +181,20 @@ void FDMXOutputPortConfigCustomization::GenerateProtocolNameRow(IDetailPropertyR
 void FDMXOutputPortConfigCustomization::GenerateCommunicationTypeRow(IDetailPropertyRow& PropertyRow)
 {
 	// Customizate the Communication Type property depending on the port's supported communication modes
-
-	TAttribute<EVisibility> CommunicationTypeVisibilityAttribute = TAttribute<EVisibility>::Create(TAttribute<EVisibility>::FGetter::CreateSP(this, &FDMXOutputPortConfigCustomization::GetCommunicationTypeVisibility));
-	PropertyRow.Visibility(CommunicationTypeVisibilityAttribute);
-
 	TSharedPtr<SWidget> NameWidget;
 	TSharedPtr<SWidget> ValueWidget;
 	FDetailWidgetRow Row;
 	PropertyRow.GetDefaultWidgets(NameWidget, ValueWidget, Row);
 
 	PropertyRow.CustomWidget()
-		.Visibility(CommunicationTypeVisibilityAttribute)
+		.Visibility(TAttribute<EVisibility>::CreateLambda([this]()
+			{
+				if (CommunicationTypeComboBox.IsValid() && GetProtocol().IsValid())
+				{
+					return CommunicationTypeComboBox->GetVisibility();
+				}
+				return EVisibility::Collapsed;
+			}))
 		.NameContent()
 		[
 			NameWidget.ToSharedRef()
@@ -189,7 +208,56 @@ void FDMXOutputPortConfigCustomization::GenerateCommunicationTypeRow(IDetailProp
 		];
 }
 
-void FDMXOutputPortConfigCustomization::GenerateIPAddressRow(IDetailPropertyRow& PropertyRow)
+void FDMXOutputPortConfigCustomization::GenerateAutoCompleteDeviceAddressRow(IDetailPropertyRow& PropertyRow)
+{
+	// Mimic a Inline Edit Condition Toggle - Doing it directly via Meta Specifiers is not possible. -
+	// The Inline Edit Condition Toggle property wouldn't get properly serialized when being a Config property.
+	const TSharedRef<SWidget> AutoCompleteDeviceAddressEnableldPropertyValueWidget = AutoCompleteDeviceAddressEnabledHandle->CreatePropertyValueWidget();
+
+	const TSharedRef<SWidget> AutoCompleteDeviceAddressPropertyNameWidget = AutoCompleteDeviceAddressHandle->CreatePropertyNameWidget();
+	const TSharedRef<SWidget> AutoCompleteDeviceAddressPropertyValueWidget = AutoCompleteDeviceAddressHandle->CreatePropertyValueWidget();
+
+	AutoCompleteDeviceAddressPropertyNameWidget->SetEnabled(TAttribute<bool>::CreateLambda([this]()
+		{
+			return IsAutoCompleteDeviceAddressEnabled();
+		}));
+	AutoCompleteDeviceAddressPropertyValueWidget->SetEnabled(TAttribute<bool>::CreateLambda([this]()
+		{
+			return IsAutoCompleteDeviceAddressEnabled();
+		}));
+
+	PropertyRow.CustomWidget()
+		.NameContent()
+		[
+			SNew(SHorizontalBox)
+
+		// Edit Inline Condition Toggle Checkbox
+		+ SHorizontalBox::Slot()
+		.HAlign(HAlign_Left)
+		.VAlign(VAlign_Fill)
+		.AutoWidth()
+		[
+			AutoCompleteDeviceAddressEnableldPropertyValueWidget
+		]
+
+	// Property Name
+	+ SHorizontalBox::Slot()
+		.Padding(8.f, 0.f, 0.f, 0.f)
+		.HAlign(HAlign_Left)
+		.VAlign(VAlign_Fill)
+		.AutoWidth()
+		[
+			AutoCompleteDeviceAddressPropertyNameWidget
+		]
+		]
+	.ValueContent()
+		[
+			// Property Value
+			AutoCompleteDeviceAddressPropertyValueWidget
+		];
+}
+
+void FDMXOutputPortConfigCustomization::GenerateDeviceAddressRow(IDetailPropertyRow& PropertyRow)
 {
 	// Customizate the IPAddress property to show a combo box with available IP addresses
 	FString InitialValue = GetIPAddress();
@@ -199,6 +267,48 @@ void FDMXOutputPortConfigCustomization::GenerateIPAddressRow(IDetailPropertyRow&
 	FDetailWidgetRow Row;
 	PropertyRow.GetDefaultWidgets(NameWidget, ValueWidget, Row);
 
+	// Override the value widget so it either shows 
+	// - The IP Address Edit widget when auto complete is disabled
+	// - The auto completed IP Address when auto complete is enabled
+	ValueWidget =
+		SNew(SHorizontalBox)
+		
+		+ SHorizontalBox::Slot()
+		.HAlign(HAlign_Fill)
+		.VAlign(VAlign_Fill)
+		[
+			SAssignNew(IPAddressEditWidget, SDMXIPAddressEditWidget)
+			.Visibility_Lambda([this]()
+				{
+					return IsAutoCompleteDeviceAddressEnabled() ?
+						EVisibility::Collapsed :
+						EVisibility::Visible;
+				})
+			.InitialValue(InitialValue)
+			.bShowLocalNICComboBox(true)
+			.OnIPAddressSelected(this, &FDMXOutputPortConfigCustomization::OnDeviceAddressSelected)
+		]
+
+		+ SHorizontalBox::Slot()
+		.HAlign(HAlign_Fill)
+		.VAlign(VAlign_Fill)
+		[
+			SAssignNew(AutoCompletedDeviceAddressTextBlock, STextBlock)
+			.Visibility_Lambda([this]()
+				{
+					return IsAutoCompleteDeviceAddressEnabled() ?
+						EVisibility::Visible :
+						EVisibility::Collapsed;
+				})
+			.Font(FEditorStyle::GetFontStyle(TEXT("PropertyWindow.NormalFont")))
+		];
+	UpdateAutoCompleteDeviceAddressTextBox();
+
+	PropertyRow.IsEnabled(TAttribute<bool>::CreateLambda([this]
+		{
+			return !IsAutoCompleteDeviceAddressEnabled();
+		}));
+
 	PropertyRow.CustomWidget()
 		.NameContent()
 		[
@@ -206,10 +316,7 @@ void FDMXOutputPortConfigCustomization::GenerateIPAddressRow(IDetailPropertyRow&
 		]
 		.ValueContent()
 		[
-			SAssignNew(IPAddressEditWidget, SDMXIPAddressEditWidget)
-			.InitialValue(InitialValue)
-			.bShowLocalNICComboBox(true)
-			.OnIPAddressSelected(this, &FDMXOutputPortConfigCustomization::OnIPAddressSelected)
+			ValueWidget.ToSharedRef()
 		];
 }	
 
@@ -245,14 +352,56 @@ void FDMXOutputPortConfigCustomization::GenerateDelayRow(IDetailPropertyRow& Pro
 		];
 }
 
-EVisibility FDMXOutputPortConfigCustomization::GetCommunicationTypeVisibility() const
+void FDMXOutputPortConfigCustomization::UpdateAutoCompleteDeviceAddressTextBox()
 {
-	if (CommunicationTypeComboBox.IsValid() && GetProtocol().IsValid())
+	// Set text on the AutoCompleteDeviceAddress Text Box depending on what the resulting IP is
+	if (AutoCompletedDeviceAddressTextBlock.IsValid())
 	{
-return CommunicationTypeComboBox->GetVisibility();
-	}
+		FString AutoCompleteDeviceAddress;
+		if (AutoCompleteDeviceAddressHandle->GetValue(AutoCompleteDeviceAddress) == FPropertyAccess::Success)
+		{
+			TArray<FString> SearchStrings;
+			constexpr bool bCullEmpty = true;
+			AutoCompleteDeviceAddress.ParseIntoArray(SearchStrings, TEXT("."), bCullEmpty);
 
-	return EVisibility::Collapsed;
+			const TArray<TSharedPtr<FString>> LocalNetworkInterfaceCardIPs = FDMXProtocolUtils::GetLocalNetworkInterfaceCardIPs();
+
+			FString NewIPAddress;
+			for (const TSharedPtr<FString>& NetworkInterfaceCardIP : LocalNetworkInterfaceCardIPs)
+			{
+				TArray<FString> NetworkInterfaceCardIPSubstrings;
+				(*NetworkInterfaceCardIP).ParseIntoArray(NetworkInterfaceCardIPSubstrings, TEXT("."));
+
+				// Try to match each search string
+				bool bMatchesSearchStrings = true;
+				for (int32 SubstringIndex = 0; SubstringIndex < NetworkInterfaceCardIPSubstrings.Num(); SubstringIndex++)
+				{
+					if (SearchStrings.IsValidIndex(SubstringIndex))
+					{
+						if (!NetworkInterfaceCardIPSubstrings[SubstringIndex].MatchesWildcard(SearchStrings[SubstringIndex]))
+						{
+							bMatchesSearchStrings = false;
+							break;
+						}
+					}
+				}
+
+				if (bMatchesSearchStrings)
+				{
+					// Set the text and return if a match was found
+					const FText IPAddressText = FText::FromString(*NetworkInterfaceCardIP);
+					AutoCompletedDeviceAddressTextBlock->SetText(FText::Format(LOCTEXT("AutoCompletedIPAddressText", "{0} (auto-completed)"), IPAddressText));
+					return;
+				}
+			}
+
+			// Fall back to the manually set one
+			FString DeviceAddress;
+			DeviceAddressHandle->GetValue(DeviceAddress);
+			const FText NoMatchForIPAddressText = FText::Format(LOCTEXT("NoRegexMatchForDeviceAddress", "No match. Using '{0}'"), FText::FromString(DeviceAddress));
+			AutoCompletedDeviceAddressTextBlock->SetText(NoMatchForIPAddressText);
+		}
+	}
 }
 
 void FDMXOutputPortConfigCustomization::OnProtocolNameSelected()
@@ -281,7 +430,7 @@ void FDMXOutputPortConfigCustomization::OnCommunicationTypeSelected()
 	PropertyUtilities->ForceRefresh();
 }
 
-void FDMXOutputPortConfigCustomization::OnIPAddressSelected()
+void FDMXOutputPortConfigCustomization::OnDeviceAddressSelected()
 {
 	FString SelectedIP = IPAddressEditWidget->GetSelectedIPAddress();
 
@@ -431,6 +580,14 @@ EDMXCommunicationType FDMXOutputPortConfigCustomization::GetCommunicationType()
 	return static_cast<EDMXCommunicationType>(CommunicationType);
 }
 
+bool FDMXOutputPortConfigCustomization::IsAutoCompleteDeviceAddressEnabled() const
+{
+	bool bAutoCompleteDeviceAddressEnabled;
+	AutoCompleteDeviceAddressEnabledHandle->GetValue(bAutoCompleteDeviceAddressEnabled);
+
+	return bAutoCompleteDeviceAddressEnabled;
+}
+
 FString FDMXOutputPortConfigCustomization::GetIPAddress() const
 {
 	FGuid PortGuid = GetPortGuid();
diff --git a/Engine/Plugins/VirtualProduction/DMX/DMXProtocol/Source/DMXProtocolEditor/Private/DetailsCustomizations/DMXOutputPortConfigCustomization.h b/Engine/Plugins/VirtualProduction/DMX/DMXProtocol/Source/DMXProtocolEditor/Private/DetailsCustomizations/DMXOutputPortConfigCustomization.h
--- a/Engine/Plugins/VirtualProduction/DMX/DMXProtocol/Source/DMXProtocolEditor/Private/DetailsCustomizations/DMXOutputPortConfigCustomization.h
+++ b/Engine/Plugins/VirtualProduction/DMX/DMXProtocol/Source/DMXProtocolEditor/Private/DetailsCustomizations/DMXOutputPortConfigCustomization.h
@@ -40,24 +40,27 @@ private:
 	/** Generates the customized Communication Type row */
 	void GenerateCommunicationTypeRow(IDetailPropertyRow& PropertyRow);
 
+	/** Generates the customized Auto Complete Device Address row */
+	void GenerateAutoCompleteDeviceAddressRow(IDetailPropertyRow& PropertyRow);
+
 	/** Generates the customized IP Address row */
-	void GenerateIPAddressRow(IDetailPropertyRow& PropertyRow);
+	void GenerateDeviceAddressRow(IDetailPropertyRow& PropertyRow);
 
 	/** Generates the customized Delay Seconds row */
 	void GenerateDelayRow(IDetailPropertyRow& PropertyRow);
 
-	/** Returns the visibility of the Communication Type Property */
-	EVisibility GetCommunicationTypeVisibility() const;
+	/** Updates the Auto Complete Device Address Text Box */
+	void UpdateAutoCompleteDeviceAddressTextBox();
 
 	/** Called when a Protocol Name was selected */
 	void OnProtocolNameSelected();
 
-	/** Called when a local IP Address was selected */
-	void OnIPAddressSelected();
-
 	/** Called when a Communication Type was selected */
 	void OnCommunicationTypeSelected();
 
+	/** Called when a Device Address was selected */
+	void OnDeviceAddressSelected();
+
 	/** Called when the Destination Addresses Array visibility needs to be updated */
 	void UpdateDestinationAddressesVisibility();
 
@@ -85,12 +88,21 @@ private:
 	/** Gets the communication type */
 	EDMXCommunicationType GetCommunicationType() const;
 
+	/** Returns the value of the AutoCompleteDeviceAddressEnabled property */
+	bool IsAutoCompleteDeviceAddressEnabled() const;
+
 	/** Gets the IP Address */
 	FString GetIPAddress() const;
 
 	/** Property handle to the ProtocolName property */
 	TSharedPtr<IPropertyHandle> ProtocolNameHandle;
 
+	/** Property handle to the bAutoCompleteDeviceAddressEnabled property */
+	TSharedPtr<IPropertyHandle> AutoCompleteDeviceAddressEnabledHandle;
+
+	/** Property handle to the AutoCompleteDeviceAddress property */
+	TSharedPtr<IPropertyHandle> AutoCompleteDeviceAddressHandle;
+
 	/** Property handle to the IPAddress property */
 	TSharedPtr<IPropertyHandle> DeviceAddressHandle;
 
@@ -112,12 +124,15 @@ private:
 	/** ComboBox to select a protocol name */
 	TSharedPtr<SDMXProtocolNameComboBox> ProtocolNameComboBox;
 
-	/** ComboBox that displays local ip addresses */
-	TSharedPtr<SDMXIPAddressEditWidget> IPAddressEditWidget;
-
 	/** ComboBox that exposes a selection of communication types to the user */
 	TSharedPtr<SDMXCommunicationTypeComboBox> CommunicationTypeComboBox;
 
+	/** Text Block that displays the device address when it is auto completed */
+	TSharedPtr<STextBlock> AutoCompletedDeviceAddressTextBlock;
+
+	/** ComboBox that displays local ip addresses */
+	TSharedPtr<SDMXIPAddressEditWidget> IPAddressEditWidget;
+
 	/** Custom widget that lets the user set the delay */
 	TSharedPtr<SDMXDelayEditWidget> DelayEditWidget;
 
