DMX - Fix CIS NoPCH Compile LyraEditor Win64

diff --git a/Engine/Plugins/VirtualProduction/DMX/DMXProtocol/Source/DMXProtocolEditor/Private/DetailsCustomizations/DMXInputPortConfigCustomization.h b/Engine/Plugins/VirtualProduction/DMX/DMXProtocol/Source/DMXProtocolEditor/Private/DetailsCustomizations/DMXInputPortConfigCustomization.h
--- a/Engine/Plugins/VirtualProduction/DMX/DMXProtocol/Source/DMXProtocolEditor/Private/DetailsCustomizations/DMXInputPortConfigCustomization.h
+++ b/Engine/Plugins/VirtualProduction/DMX/DMXProtocol/Source/DMXProtocolEditor/Private/DetailsCustomizations/DMXInputPortConfigCustomization.h
@@ -16,6 +16,7 @@ class SDMXProtocolNameComboBox;
 class IDetailPropertyRow;
 class IPropertyHandle;
 class IPropertyUtilities;
+class STextBlock;
 
 
 /** Details customization for input and output port configs. */
diff --git a/Engine/Plugins/VirtualProduction/DMX/DMXProtocol/Source/DMXProtocolEditor/Private/DetailsCustomizations/DMXOutputPortConfigCustomization.h b/Engine/Plugins/VirtualProduction/DMX/DMXProtocol/Source/DMXProtocolEditor/Private/DetailsCustomizations/DMXOutputPortConfigCustomization.h
--- a/Engine/Plugins/VirtualProduction/DMX/DMXProtocol/Source/DMXProtocolEditor/Private/DetailsCustomizations/DMXOutputPortConfigCustomization.h
+++ b/Engine/Plugins/VirtualProduction/DMX/DMXProtocol/Source/DMXProtocolEditor/Private/DetailsCustomizations/DMXOutputPortConfigCustomization.h
@@ -18,6 +18,7 @@ struct FFrameRate;
 class IDetailPropertyRow;
 class IPropertyHandle;
 class IPropertyUtilities;
+class STextBlock;
 
 
 /** Details customization for input and output port configs. */
