DMX - Fix CIS Static Analysis Win64 (PVS-Studio)

diff --git a/Engine/Plugins/VirtualProduction/DMX/DMXProtocol/Source/DMXProtocol/Private/IO/DMXInputPort.cpp b/Engine/Plugins/VirtualProduction/DMX/DMXProtocol/Source/DMXProtocol/Private/IO/DMXInputPort.cpp
--- a/Engine/Plugins/VirtualProduction/DMX/DMXProtocol/Source/DMXProtocol/Private/IO/DMXInputPort.cpp
+++ b/Engine/Plugins/VirtualProduction/DMX/DMXProtocol/Source/DMXProtocol/Private/IO/DMXInputPort.cpp
@@ -455,8 +455,8 @@ void FDMXInputPort::UpdateFromConfig(FDMXInputPortConfig& InOutInputPortConfig,
 		FName ProtocolName = Protocol.IsValid() ? Protocol->GetProtocolName() : NAME_None;
 
 		if (ProtocolName == InputPortConfig.GetProtocolName() &&
-			bAutoCompleteDeviceAddressEnabled == bAutoCompleteDeviceAddressEnabled &&
-			AutoCompleteDeviceAddress == AutoCompleteDeviceAddress &&
+			bAutoCompleteDeviceAddressEnabled == InputPortConfig.IsAutoCompleteDeviceAddressEnabled() &&
+			AutoCompleteDeviceAddress == InputPortConfig.GetAutoCompleteDeviceAddress() &&
 			DeviceAddress == InputPortConfig.GetDeviceAddress() &&
 			CommunicationType == InputPortConfig.GetCommunicationType())
 		{
diff --git a/Engine/Plugins/VirtualProduction/DMX/DMXProtocol/Source/DMXProtocol/Private/IO/DMXOutputPort.cpp b/Engine/Plugins/VirtualProduction/DMX/DMXProtocol/Source/DMXProtocol/Private/IO/DMXOutputPort.cpp
--- a/Engine/Plugins/VirtualProduction/DMX/DMXProtocol/Source/DMXProtocol/Private/IO/DMXOutputPort.cpp
+++ b/Engine/Plugins/VirtualProduction/DMX/DMXProtocol/Source/DMXProtocol/Private/IO/DMXOutputPort.cpp
@@ -577,8 +577,8 @@ void FDMXOutputPort::UpdateFromConfig(FDMXOutputPortConfig& InOutOutputPortConfi
 		FName ProtocolName = Protocol.IsValid() ? Protocol->GetProtocolName() : NAME_None;
 
 		if (ProtocolName == OutputPortConfig.GetProtocolName() &&
-			bAutoCompleteDeviceAddressEnabled == bAutoCompleteDeviceAddressEnabled &&
-			AutoCompleteDeviceAddress == AutoCompleteDeviceAddress &&
+			bAutoCompleteDeviceAddressEnabled == OutputPortConfig.IsAutoCompleteDeviceAddressEnabled() &&
+			AutoCompleteDeviceAddress == OutputPortConfig.GetAutoCompleteDeviceAddress() &&
 			DeviceAddress == OutputPortConfig.GetDeviceAddress() &&
 			DestinationAddresses == OutputPortConfig.GetDestinationAddresses() &&
 			CommunicationType == OutputPortConfig.GetCommunicationType() &&
