Discard sACN packet with PerAddressPriority StartCode (0xDD)

diff --git a/Engine/Plugins/VirtualProduction/DMX/DMXProtocol/Source/DMXProtocolSACN/Private/DMXProtocolSACNConstants.h b/Engine/Plugins/VirtualProduction/DMX/DMXProtocol/Source/DMXProtocolSACN/Private/DMXProtocolSACNConstants.h
--- a/Engine/Plugins/VirtualProduction/DMX/DMXProtocol/Source/DMXProtocolSACN/Private/DMXProtocolSACNConstants.h
+++ b/Engine/Plugins/VirtualProduction/DMX/DMXProtocol/Source/DMXProtocolSACN/Private/DMXProtocolSACNConstants.h
@@ -63,3 +63,8 @@ enum
 	VHD_H_FLAG = 0x20,
 	VHD_D_FLAG = 0x10
 };
+
+enum
+{
+	ACN_DMX_START_CODE_PER_ADDRESS_PRIORITY = 0xDD
+};
diff --git a/Engine/Plugins/VirtualProduction/DMX/DMXProtocol/Source/DMXProtocolSACN/Private/DMXProtocolSACNReceiver.cpp b/Engine/Plugins/VirtualProduction/DMX/DMXProtocol/Source/DMXProtocolSACN/Private/DMXProtocolSACNReceiver.cpp
--- a/Engine/Plugins/VirtualProduction/DMX/DMXProtocol/Source/DMXProtocolSACN/Private/DMXProtocolSACNReceiver.cpp
+++ b/Engine/Plugins/VirtualProduction/DMX/DMXProtocol/Source/DMXProtocolSACN/Private/DMXProtocolSACNReceiver.cpp
@@ -313,6 +313,18 @@ void FDMXProtocolSACNReceiver::HandleDataPacket(uint16 UniverseID, const TShared
 		return;
 	}
 
+	// discard per-address-priority StartCode (even if in the future we may want to support it)
+	if (IncomingDMXDMPLayer.STARTCode == ACN_DMX_START_CODE_PER_ADDRESS_PRIORITY)
+	{
+		static bool PerAddressPriorityUnsupportedLogged = false;
+		if (!PerAddressPriorityUnsupportedLogged)
+		{
+			UE_LOG(LogDMXProtocol, Warning, TEXT("Received DMX sACN packet(s) that specify a per address priority. This is not supported in the current Engine version."));
+			PerAddressPriorityUnsupportedLogged = true;
+		}
+		return;
+	}
+
 	// check if the number of properties is consistent
 	if ((IncomingDMXDMPLayer.FirstPropertyAddress + IncomingDMXDMPLayer.PropertyValueCount) > (ACN_DMX_SIZE + 1))
 	{
