DMX - Console command DMX.SetInputPortAddress ? resp. DMX.SetOutputPortAddress now shows all IPs

diff --git a/Engine/Plugins/VirtualProduction/DMX/DMXProtocol/Source/DMXProtocol/Private/IO/DMXInputPort.cpp b/Engine/Plugins/VirtualProduction/DMX/DMXProtocol/Source/DMXProtocol/Private/IO/DMXInputPort.cpp
--- a/Engine/Plugins/VirtualProduction/DMX/DMXProtocol/Source/DMXProtocol/Private/IO/DMXInputPort.cpp
+++ b/Engine/Plugins/VirtualProduction/DMX/DMXProtocol/Source/DMXProtocol/Private/IO/DMXInputPort.cpp
@@ -17,8 +17,48 @@ DECLARE_CYCLE_STAT(TEXT("Input Port Tick"), STAT_DMXInputPortTick, STATGROUP_DMX
 
 #define LOCTEXT_NAMESPACE "DMXInputPort"
 
-namespace UE::DMX::DMXInputPort
+/** Helper to override a member variable of an Input Port */
+#define DMX_OVERRIDE_INPUTPORT_VAR(MemberName, PortName, Value) \
+{ \
+	const FDMXInputPortSharedPtr OverRideInputPortMacro_InputPort = UE::DMX::DMXInputPort::Private::ConsoleCommands::FindInputPortByName(PortName); \
+	if (OverRideInputPortMacro_InputPort.IsValid()) \
+	{ \
+		const FDMXInputPortConfig OldInputPortConfig = OverRideInputPortMacro_InputPort->MakeInputPortConfig(); \
+		FDMXInputPortConfigParams InputPortConfigParams = FDMXInputPortConfigParams(OldInputPortConfig); \
+		InputPortConfigParams.MemberName = Value; \
+		FDMXInputPortConfig NewInputPortConfig(OverRideInputPortMacro_InputPort->GetPortGuid(), InputPortConfigParams); \
+		constexpr bool bForceUpdateRegistrationWithProtocol = true; \
+		OverRideInputPortMacro_InputPort->UpdateFromConfig(NewInputPortConfig, bForceUpdateRegistrationWithProtocol); \
+	} \
+}
+
+namespace UE::DMX::DMXInputPort::Private::ConsoleCommands
 {
+	void LogInputPortConfiguration(const FDMXInputPortSharedPtr InputPort)
+	{
+		if (InputPort.IsValid())
+		{
+			UE_LOG(LogDMXProtocol, Log, TEXT("================================================"));
+			UE_LOG(LogDMXProtocol, Log, TEXT("Configuration of Input Port '%s':"), *InputPort->GetPortName());
+			UE_LOG(LogDMXProtocol, Log, TEXT(" "));
+
+			UE_LOG(LogDMXProtocol, Log, TEXT("Port Name: %s"), *InputPort->GetPortName());
+
+			FDMXInputPortConfig PortConfig = InputPort->MakeInputPortConfig();
+			UE_LOG(LogDMXProtocol, Log, TEXT("Protocol: %s"), *PortConfig.GetProtocolName().ToString());
+			UE_LOG(LogDMXProtocol, Log, TEXT("Communication Type: %i"), static_cast<int32>(PortConfig.GetCommunicationType()));
+			UE_LOG(LogDMXProtocol, Log, TEXT("Device Address: %s"), *PortConfig.GetDeviceAddress());
+
+			UE_LOG(LogDMXProtocol, Log, TEXT("Local Universe Start: %i"), PortConfig.GetLocalUniverseStart());
+			UE_LOG(LogDMXProtocol, Log, TEXT("Num Universes: %i"), PortConfig.GetNumUniverses());
+			UE_LOG(LogDMXProtocol, Log, TEXT("Extern Universe Start: %i"), PortConfig.GetExternUniverseStart());
+			UE_LOG(LogDMXProtocol, Log, TEXT("Priority Strategy: %i"), static_cast<int32>(PortConfig.GetPortPriorityStrategy()));
+			UE_LOG(LogDMXProtocol, Log, TEXT("Priority:	%i"), PortConfig.GetPriority());
+
+			UE_LOG(LogDMXProtocol, Log, TEXT("================================================"));
+		}
+	}
+
 	/** Finds an Input port by its name, optionally logging when it can't be found */
 	static const FDMXInputPortSharedPtr FindInputPortByName(const FString& PortName, bool bPrintToLogIfNotFound = true)
 	{
@@ -33,7 +73,14 @@ namespace UE::DMX::DMXInputPort
 		}
 		else if (bPrintToLogIfNotFound)
 		{
-			UE_LOG(LogDMXProtocol, Warning, TEXT("Could not find Input Port '%s'. Available Input Ports and their names can be found in in Project settings -> Plugins -> DMX."), *PortName);
+			if (PortName.IsEmpty())
+			{
+				UE_LOG(LogDMXProtocol, Warning, TEXT("No Input Port specified for console command. First argument needs to be a valid Port Name."));
+			}
+			else
+			{
+				UE_LOG(LogDMXProtocol, Warning, TEXT("Could not find Input Port '%s'. Available Input Ports and their names can be found in in Project settings -> Plugins -> DMX."), *PortName);
+			}
 		}
 
 		return nullptr;
@@ -42,329 +89,329 @@ namespace UE::DMX::DMXInputPort
 	/** Tests the console command, returns true if they yield a valid port. Logs issues otherwise. Handles ? commands. */
 	static bool VerifyConsoleCommandArguments(const FString& ConsoleCommand, int32 MinNumExpectedArgs, const TArray<FString>& Args)
 	{
-		check(MinNumExpectedArgs >= 1);
+		const FString PortName = Args.Num() > 0 ? Args[0] : FString();
+		const FDMXInputPortSharedPtr InputPort = FindInputPortByName(PortName);
 
-		if (Args.Num() < MinNumExpectedArgs)
+		// Log the port if the there is only one arg or the second arg is '?' (e.g. 'DMX.SetInputPortProtocol InputPortA' or 'DMX.SetInputPortProtocol InputPortA ?')
+		const bool bOnlyPortSpecified = Args.Num() == 1;
+		const bool bSecondArgIsQuestionmark = (Args.Num() == 2 && Args[1] == TEXT("?"));
+		if (InputPort.IsValid() && (bOnlyPortSpecified || bSecondArgIsQuestionmark))
 		{
-			UE_LOG(LogDMXProtocol, Warning, TEXT("Console Command %s failed. Insufficient arguments specified."), *ConsoleCommand);
-			UE_LOG(LogDMXProtocol, Warning, TEXT("Use '%s ?' for help. Use '%s PortName ?' to print the current configuration to logs."), *ConsoleCommand, *ConsoleCommand);
+			LogInputPortConfiguration(InputPort);
+			
 			return false;
 		}
-
-		const FString& PortName = Args[0];
-		const FDMXInputPortSharedPtr InputPort = UE::DMX::DMXInputPort::FindInputPortByName(PortName);
-		if (InputPort.IsValid())
+		else if (!InputPort.IsValid())
 		{
-			// Log the port if the console command sets the second arg to '?' (e.g. 'DMX.SetInputPortProtocol InputPortA ?')
-			if (Args.Num() == 2 && Args[1] == TEXT("?"))
-			{
-				UE_LOG(LogDMXProtocol, Log, TEXT("Input Port: %s"), *InputPort->GetPortName());
-
-				FDMXInputPortConfig PortConfig = InputPort->MakeInputPortConfig();
-				UE_LOG(LogDMXProtocol, Log, TEXT("Protocol: %s"), *PortConfig.GetProtocolName().ToString());
-				UE_LOG(LogDMXProtocol, Log, TEXT("Communication Type: %i"), static_cast<int32>(PortConfig.GetCommunicationType()));
-				UE_LOG(LogDMXProtocol, Log, TEXT("Device Address: %s"), *PortConfig.GetDeviceAddress());
-
-				UE_LOG(LogDMXProtocol, Log, TEXT("Local Universe Start: %i"), PortConfig.GetLocalUniverseStart());
-				UE_LOG(LogDMXProtocol, Log, TEXT("Num Universes: %i"), PortConfig.GetNumUniverses());
-				UE_LOG(LogDMXProtocol, Log, TEXT("Extern Universe Start: %i"), PortConfig.GetExternUniverseStart());
-				UE_LOG(LogDMXProtocol, Log, TEXT("Priority Strategy: %i"), static_cast<int32>(PortConfig.GetPortPriorityStrategy()));
-				UE_LOG(LogDMXProtocol, Log, TEXT("Priority:	%i"), PortConfig.GetPriority());
-			}
-			else
-			{
-				return true;
-			}
+			// Log additional help if there were not enough arguments or arguments were invalid
+			UE_LOG(LogDMXProtocol, Warning, TEXT("Console Command %s failed. Input Port with Name '%s' does not exist."), *ConsoleCommand, *PortName);
+			UE_LOG(LogDMXProtocol, Warning, TEXT("Ports and their Names can be found in Project Settings -> Plugins -> DMX."), *ConsoleCommand, *ConsoleCommand);
+			
+			return false;
+		}
+		else if (MinNumExpectedArgs > Args.Num())
+		{
+			UE_LOG(LogDMXProtocol, Warning, TEXT("Console Command %s failed. Invalid arguments specified."), *ConsoleCommand);
+			UE_LOG(LogDMXProtocol, Warning, TEXT("Use '%s ?' for help. Use '%s PortName ?' to print the current configuration to logs."), *ConsoleCommand, *ConsoleCommand);
+			
+			return false;
 		}
 
-		return false;
-	}
-}
 
-/** Helper to override a member variable of an Input Port */
-#define DMX_OVERRIDE_INPUTPORT_VAR(MemberName, PortName, Value) \
-{ \
-	const FDMXInputPortSharedPtr OverRideInputPortMacro_InputPort = UE::DMX::DMXInputPort::FindInputPortByName(PortName); \
-	if (OverRideInputPortMacro_InputPort.IsValid()) \
-	{ \
-		const FDMXInputPortConfig OldInputPortConfig = OverRideInputPortMacro_InputPort->MakeInputPortConfig(); \
-		FDMXInputPortConfigParams InputPortConfigParams = FDMXInputPortConfigParams(OldInputPortConfig); \
-		InputPortConfigParams.MemberName = Value; \
-		FDMXInputPortConfig NewInputPortConfig(OverRideInputPortMacro_InputPort->GetPortGuid(), InputPortConfigParams); \
-		constexpr bool bForceUpdateRegistrationWithProtocol = true; \
-		OverRideInputPortMacro_InputPort->UpdateFromConfig(NewInputPortConfig, bForceUpdateRegistrationWithProtocol); \
-	} \
-}
+		return true;
+	}
 
-static FAutoConsoleCommand GDMXSetInputPortProtocolCommand(
-	TEXT("DMX.SetInputPortProtocol"),
-	TEXT("DMX.SetInputPortProtocol [PortName][ProtocolName]. Sets the protocol used by the input port. Example: DMX.SetInputPortProtocol MyInputPort sACN"),
-	FConsoleCommandWithArgsDelegate::CreateStatic(
-		[](const TArray<FString>& Args)
-		{
-			constexpr int32 MinNumExpectedArgs = 2;
-			if (!UE::DMX::DMXInputPort::VerifyConsoleCommandArguments(TEXT("DMX.SetInputPortProtocol"), MinNumExpectedArgs, Args))
+	static FAutoConsoleCommand GDMXLogAllInputPortConfigurationsCommand(
+		TEXT("DMX.LogAllInputPortConfigurations"), 
+		TEXT("Logs all Input Port configurations"), 
+		FConsoleCommandDelegate::CreateStatic(
+			[]()
 			{
-				return;
-			}
+				for (const FDMXInputPortSharedRef& InputPort : FDMXPortManager::Get().GetInputPorts())
+				{
+					LogInputPortConfiguration(InputPort);
+				}
+			})
+	);
+
+	static FAutoConsoleCommand GDMXSetInputPortProtocolCommand(
+		TEXT("DMX.SetInputPortProtocol"),
+		TEXT("DMX.SetInputPortProtocol [PortName][ProtocolName]. Sets the protocol used by the input port. Example: DMX.SetInputPortProtocol MyInputPort sACN"),
+		FConsoleCommandWithArgsDelegate::CreateStatic(
+			[](const TArray<FString>& Args)
+			{
+				constexpr int32 MinNumExpectedArgs = 2;
+				if (!VerifyConsoleCommandArguments(TEXT("DMX.SetInputPortProtocol"), MinNumExpectedArgs, Args))
+				{
+					return;
+				}
 
-			const FString& PortName = Args[0];
-			const FName ProtocolNameValue = FName(Args[1]);
+				const FString& PortName = Args[0];
+				const FName ProtocolNameValue = FName(Args[1]);
 
-			if (IDMXProtocol::GetProtocolNames().Contains(ProtocolNameValue))
-			{
-				DMX_OVERRIDE_INPUTPORT_VAR(ProtocolName, PortName, ProtocolNameValue);
-			}
-			else
-			{
-				UE_LOG(LogDMXProtocol, Warning, TEXT("Console command DMX.SetInputPortProtocol failed. Protocol '%s' does not exist or is not loaded. Available protocols are: "), *ProtocolNameValue.ToString());
-				for (const FName& ProtocolName : IDMXProtocol::GetProtocolNames())
+				if (IDMXProtocol::GetProtocolNames().Contains(ProtocolNameValue))
 				{
-					UE_LOG(LogDMXProtocol, Warning, TEXT("%s"), *ProtocolName.ToString());
+					DMX_OVERRIDE_INPUTPORT_VAR(ProtocolName, PortName, ProtocolNameValue);
 				}
-			}
-		})
-);
-
-static FAutoConsoleCommand GDMXSetInputPortDeviceAddressCommand(
-	TEXT("DMX.SetInputPortDeviceAddress"),
-	TEXT("DMX.SetInputPortDeviceAddress [PortName][DeviceAddress]. Sets the Device Address of an input port, usually the network interface card IP address. Example: DMX.SetInputPortDeviceAddress MyInputPort 123.45.67.89"),
-	FConsoleCommandWithArgsDelegate::CreateStatic(
-		[](const TArray<FString>& Args)
-		{
-			constexpr int32 MinNumExpectedArgs = 2;
-			if (!UE::DMX::DMXInputPort::VerifyConsoleCommandArguments(TEXT("DMX.SetInputPortProtocol"), MinNumExpectedArgs, Args))
+				else
+				{
+					UE_LOG(LogDMXProtocol, Warning, TEXT("Console command DMX.SetInputPortProtocol failed. Protocol '%s' does not exist or is not loaded. Available protocols are: "), *ProtocolNameValue.ToString());
+					for (const FName& ProtocolName : IDMXProtocol::GetProtocolNames())
+					{
+						UE_LOG(LogDMXProtocol, Warning, TEXT("%s"), *ProtocolName.ToString());
+					}
+				}
+			})
+	);
+
+	static FAutoConsoleCommand GDMXSetInputPortDeviceAddressCommand(
+		TEXT("DMX.SetInputPortDeviceAddress"),
+		TEXT("DMX.SetInputPortDeviceAddress [PortName][DeviceAddress]. Sets the Device Address of an input port, usually the network interface card IP address. Example: DMX.SetInputPortDeviceAddress MyInputPort 123.45.67.89"),
+		FConsoleCommandWithArgsDelegate::CreateStatic(
+			[](const TArray<FString>& Args)
 			{
-				return;
-			}
-
-			const FString& PortName = Args[0];
-			const FString& DeviceAddressValue = Args[1];
+				constexpr int32 MinNumExpectedArgs = 2;
+				if (!VerifyConsoleCommandArguments(TEXT("DMX.SetInputPortDeviceAddress"), MinNumExpectedArgs, Args))
+				{							
+					// For this command, additionally log the available Local Network Interface Card IP addresses
+					UE_LOG(LogDMXProtocol, Log, TEXT("================================================"));
+					UE_LOG(LogDMXProtocol, Log, TEXT("Logging available Local Network Interface Card IP Addresses:"));
+					const TArray<TSharedPtr<FString>> AvailableIPAddresses = FDMXProtocolUtils::GetLocalNetworkInterfaceCardIPs();
+					for (int32 IPAddressIndex = 0; IPAddressIndex < AvailableIPAddresses.Num(); IPAddressIndex++)
+					{
+						UE_LOG(LogDMXProtocol, Log, TEXT("%i: %s"), IPAddressIndex, *(*AvailableIPAddresses[IPAddressIndex]));
+					}
+					UE_LOG(LogDMXProtocol, Log, TEXT("================================================"));
 
-			DMX_OVERRIDE_INPUTPORT_VAR(DeviceAddress, PortName, DeviceAddressValue);
-		})
-);
+					return;
+				}
 
-static FAutoConsoleCommand GDMXSetInputPortLocalUniverseStartCommand(
-	TEXT("DMX.SetInputPortLocalUniverseStart"),
-	TEXT("DMX.SetInputPortLocalUniverseStart [PortName][Universe]. Sets the local universe start of the input port. Example: DMX.SetInputPortLocalUniverseStart MyInputPort 5"),
-	FConsoleCommandWithArgsDelegate::CreateStatic(
-		[](const TArray<FString>& Args)
-		{
-			constexpr int32 MinNumExpectedArgs = 2;
-			if (!UE::DMX::DMXInputPort::VerifyConsoleCommandArguments(TEXT("DMX.SetInputPortProtocol"), MinNumExpectedArgs, Args))
-			{
-				return;
-			}
+				const FString& PortName = Args[0];
+				const FString& DeviceAddressValue = Args[1];
 
-			const FString& PortName = Args[0];
-			const FDMXInputPortSharedPtr InputPort = UE::DMX::DMXInputPort::FindInputPortByName(PortName);
+				DMX_OVERRIDE_INPUTPORT_VAR(DeviceAddress, PortName, DeviceAddressValue);
+			})
+	);
 
-			if (InputPort.IsValid())
+	static FAutoConsoleCommand GDMXSetInputPortLocalUniverseStartCommand(
+		TEXT("DMX.SetInputPortLocalUniverseStart"),
+		TEXT("DMX.SetInputPortLocalUniverseStart [PortName][Universe]. Sets the local universe start of the input port. Example: DMX.SetInputPortLocalUniverseStart MyInputPort 5"),
+		FConsoleCommandWithArgsDelegate::CreateStatic(
+			[](const TArray<FString>& Args)
 			{
-				if (const IDMXProtocolPtr& Protocol = InputPort->GetProtocol())
+				constexpr int32 MinNumExpectedArgs = 2;
+				if (!VerifyConsoleCommandArguments(TEXT("DMX.SetInputPortLocalUniverseStart"), MinNumExpectedArgs, Args))
 				{
-					const FString& LocalUniverseStartValueString = Args[1];
-					int32 LocalUniverseStartValue;
-					if (LexTryParseString<int32>(LocalUniverseStartValue, *LocalUniverseStartValueString))
-					{
-						const int32 NumUniverses = InputPort->GetLocalUniverseEnd() - InputPort->GetLocalUniverseStart() + 1;
+					return;
+				}
 
-						const int32 ExternUniverseStart = InputPort->ConvertLocalToExternUniverseID(LocalUniverseStartValue);
-						const int32 ExternUniverseEnd = InputPort->ConvertLocalToExternUniverseID(LocalUniverseStartValue + NumUniverses - 1);
+				const FString& PortName = Args[0];
+				const FDMXInputPortSharedPtr InputPort = FindInputPortByName(PortName);
 
-						if (Protocol->IsValidUniverseID(ExternUniverseStart) &&
-							Protocol->IsValidUniverseID(ExternUniverseEnd))
-						{
-							DMX_OVERRIDE_INPUTPORT_VAR(LocalUniverseStart, PortName, LocalUniverseStartValue);
-						}
-						else
+				if (InputPort.IsValid())
+				{
+					if (const IDMXProtocolPtr& Protocol = InputPort->GetProtocol())
+					{
+						const FString& LocalUniverseStartValueString = Args[1];
+						int32 LocalUniverseStartValue;
+						if (LexTryParseString<int32>(LocalUniverseStartValue, *LocalUniverseStartValueString))
 						{
-							UE_LOG(LogDMXProtocol, Warning, TEXT("Console command DMX.SetInputPortLocalUniverseStart failed. Local Universe Start '%s' along with Num Universes '%i' results in a Universe range that is not supported by the Protocol of the Port."), *LocalUniverseStartValueString, NumUniverses);
+							const int32 NumUniverses = InputPort->GetLocalUniverseEnd() - InputPort->GetLocalUniverseStart() + 1;
+
+							const int32 ExternUniverseStart = InputPort->ConvertLocalToExternUniverseID(LocalUniverseStartValue);
+							const int32 ExternUniverseEnd = InputPort->ConvertLocalToExternUniverseID(LocalUniverseStartValue + NumUniverses - 1);
+
+							if (Protocol->IsValidUniverseID(ExternUniverseStart) &&
+								Protocol->IsValidUniverseID(ExternUniverseEnd))
+							{
+								DMX_OVERRIDE_INPUTPORT_VAR(LocalUniverseStart, PortName, LocalUniverseStartValue);
+							}
+							else
+							{
+								UE_LOG(LogDMXProtocol, Warning, TEXT("Console command DMX.SetInputPortLocalUniverseStart failed. Local Universe Start '%s' along with Num Universes '%i' results in a Universe range that is not supported by the Protocol of the Port."), *LocalUniverseStartValueString, NumUniverses);
+							}
 						}
 					}
 				}
-			}
-		})
-);
-
-static FAutoConsoleCommand GDMXSetInputPortNumUniversesCommand(
-	TEXT("DMX.SetInputPortNumUniverses"),
-	TEXT("DMX.SetInputPortNumUniverses [PortName][Universe]. Sets the num universes of the input port. Example: DMX.SetInputPortNumUniverses MyInputPort 10"),
-	FConsoleCommandWithArgsDelegate::CreateStatic(
-		[](const TArray<FString>& Args)
-		{
-			constexpr int32 MinNumExpectedArgs = 2;
-			if (!UE::DMX::DMXInputPort::VerifyConsoleCommandArguments(TEXT("DMX.SetInputPortProtocol"), MinNumExpectedArgs, Args))
+			})
+	);
+
+	static FAutoConsoleCommand GDMXSetInputPortNumUniversesCommand(
+		TEXT("DMX.SetInputPortNumUniverses"),
+		TEXT("DMX.SetInputPortNumUniverses [PortName][Universe]. Sets the num universes of the input port. Example: DMX.SetInputPortNumUniverses MyInputPort 10"),
+		FConsoleCommandWithArgsDelegate::CreateStatic(
+			[](const TArray<FString>& Args)
 			{
-				return;
-			}
+				constexpr int32 MinNumExpectedArgs = 2;
+				if (!VerifyConsoleCommandArguments(TEXT("DMX.SetInputPortNumUniverses"), MinNumExpectedArgs, Args))
+				{
+					return;
+				}
 
-			const FString& PortName = Args[0];
-			const FDMXInputPortSharedPtr InputPort = UE::DMX::DMXInputPort::FindInputPortByName(PortName);
+				const FString& PortName = Args[0];
+				const FDMXInputPortSharedPtr InputPort = FindInputPortByName(PortName);
 
-			if (InputPort.IsValid())
-			{
-				if (const IDMXProtocolPtr& Protocol = InputPort->GetProtocol())
+				if (InputPort.IsValid())
 				{
-					const FString& NumUniversesValueString = Args[1];
-					int32 NumUniversesValue;
-					if (LexTryParseString<int32>(NumUniversesValue, *NumUniversesValueString))
+					if (const IDMXProtocolPtr& Protocol = InputPort->GetProtocol())
 					{
-						const int32 LocalUniverseStart = InputPort->GetLocalUniverseStart();
-
-						const int32 ExternUniverseStart = InputPort->ConvertLocalToExternUniverseID(LocalUniverseStart);
-						const int32 ExternUniverseEnd = ExternUniverseStart + NumUniversesValue - 1;
-
-						if (NumUniversesValue > 0 &&
-							Protocol->IsValidUniverseID(ExternUniverseStart) &&
-							Protocol->IsValidUniverseID(ExternUniverseEnd))
-						{
-							DMX_OVERRIDE_INPUTPORT_VAR(NumUniverses, PortName, NumUniversesValue);
-						}
-						else
+						const FString& NumUniversesValueString = Args[1];
+						int32 NumUniversesValue;
+						if (LexTryParseString<int32>(NumUniversesValue, *NumUniversesValueString))
 						{
-							UE_LOG(LogDMXProtocol, Warning, TEXT("Console command DMX.SetInputPortNumUniverses failed. Local Universe Start '%i' along with Num Universes '%s' results in a Universe range that is not supported by the Protocol of the Port."), LocalUniverseStart, *NumUniversesValueString);
+							const int32 LocalUniverseStart = InputPort->GetLocalUniverseStart();
+
+							const int32 ExternUniverseStart = InputPort->ConvertLocalToExternUniverseID(LocalUniverseStart);
+							const int32 ExternUniverseEnd = ExternUniverseStart + NumUniversesValue - 1;
+
+							if (NumUniversesValue > 0 &&
+								Protocol->IsValidUniverseID(ExternUniverseStart) &&
+								Protocol->IsValidUniverseID(ExternUniverseEnd))
+							{
+								DMX_OVERRIDE_INPUTPORT_VAR(NumUniverses, PortName, NumUniversesValue);
+							}
+							else
+							{
+								UE_LOG(LogDMXProtocol, Warning, TEXT("Console command DMX.SetInputPortNumUniverses failed. Local Universe Start '%i' along with Num Universes '%s' results in a Universe range that is not supported by the Protocol of the Port."), LocalUniverseStart, *NumUniversesValueString);
+							}
 						}
 					}
 				}
-			}
-		})
-);
-
-static FAutoConsoleCommand GDMXSetInputPortExternUniverseStartCommand(
-	TEXT("DMX.SetInputPortExternUniverseStart"),
-	TEXT("DMX.SetInputPortExternUniverseStart [PortName][Universe]. Sets the extern universe start of the input port. Example: DMX.SetInputPortExternUniverseStart MyInputPort 7"),
-	FConsoleCommandWithArgsDelegate::CreateStatic(
-		[](const TArray<FString>& Args)
-		{
-			constexpr int32 MinNumExpectedArgs = 2;
-			if (!UE::DMX::DMXInputPort::VerifyConsoleCommandArguments(TEXT("DMX.SetInputPortProtocol"), MinNumExpectedArgs, Args))
+			})
+	);
+
+	static FAutoConsoleCommand GDMXSetInputPortExternUniverseStartCommand(
+		TEXT("DMX.SetInputPortExternUniverseStart"),
+		TEXT("DMX.SetInputPortExternUniverseStart [PortName][Universe]. Sets the extern universe start of the input port. Example: DMX.SetInputPortExternUniverseStart MyInputPort 7"),
+		FConsoleCommandWithArgsDelegate::CreateStatic(
+			[](const TArray<FString>& Args)
 			{
-				return;
-			}
+				constexpr int32 MinNumExpectedArgs = 2;
+				if (!VerifyConsoleCommandArguments(TEXT("DMX.SetInputPortExternUniverseStart"), MinNumExpectedArgs, Args))
+				{
+					return;
+				}
 
-			const FString& PortName = Args[0];
-			const FDMXInputPortSharedPtr InputPort = UE::DMX::DMXInputPort::FindInputPortByName(PortName);
+				const FString& PortName = Args[0];
+				const FDMXInputPortSharedPtr InputPort = FindInputPortByName(PortName);
 
-			if (InputPort.IsValid())
-			{
-				const FString& ExternUniverseStartValueString = Args[1];
-				int32 ExternUniverseStartValue;
-				if (LexTryParseString<int32>(ExternUniverseStartValue, *ExternUniverseStartValueString))
+				if (InputPort.IsValid())
 				{
-					if (const IDMXProtocolPtr& Protocol = InputPort->GetProtocol())
+					const FString& ExternUniverseStartValueString = Args[1];
+					int32 ExternUniverseStartValue;
+					if (LexTryParseString<int32>(ExternUniverseStartValue, *ExternUniverseStartValueString))
 					{
-						const int32 NumUniverses = InputPort->GetLocalUniverseEnd() - InputPort->GetLocalUniverseStart() + 1;
-						const int32 ExternUniverseEnd = ExternUniverseStartValue + NumUniverses - 1;
-
-						if (Protocol->IsValidUniverseID(ExternUniverseStartValue) &&
-							Protocol->IsValidUniverseID(ExternUniverseEnd))
+						if (const IDMXProtocolPtr& Protocol = InputPort->GetProtocol())
 						{
-							DMX_OVERRIDE_INPUTPORT_VAR(ExternUniverseStart, PortName, ExternUniverseStartValue);
-						}
-						else
-						{
-							UE_LOG(LogDMXProtocol, Warning, TEXT("Console command DMX.SetInputPortExternUniverseStart failed. Extern Universe Start '%s' along with Num Universes '%i' results in a Universe range that is not supported by the Protocol of the Port."), *ExternUniverseStartValueString, NumUniverses);
+							const int32 NumUniverses = InputPort->GetLocalUniverseEnd() - InputPort->GetLocalUniverseStart() + 1;
+							const int32 ExternUniverseEnd = ExternUniverseStartValue + NumUniverses - 1;
+
+							if (Protocol->IsValidUniverseID(ExternUniverseStartValue) &&
+								Protocol->IsValidUniverseID(ExternUniverseEnd))
+							{
+								DMX_OVERRIDE_INPUTPORT_VAR(ExternUniverseStart, PortName, ExternUniverseStartValue);
+							}
+							else
+							{
+								UE_LOG(LogDMXProtocol, Warning, TEXT("Console command DMX.SetInputPortExternUniverseStart failed. Extern Universe Start '%s' along with Num Universes '%i' results in a Universe range that is not supported by the Protocol of the Port."), *ExternUniverseStartValueString, NumUniverses);
+							}
 						}
 					}
 				}
-			}
-		})
-);
-
-static FAutoConsoleCommand GDMXSetInputPortPriorityStrategyCommand(
-	TEXT("DMX.SetInputPortPriorityStrategy"),
-	TEXT("DMX.SetInputPortPriorityStrategy [PortName][PriorityStrategy (0 = None, 1 = Lowest, 2 = LowerThan, 3 = Equal, 4 = Higher Than, 5 = Highest)]. Example: DMX.SetInputPortPriorityStrategy MyInputPort 1"),
-	FConsoleCommandWithArgsDelegate::CreateStatic(
-		[](const TArray<FString>& Args)
-		{
-			constexpr int32 MinNumExpectedArgs = 2;
-			if (!UE::DMX::DMXInputPort::VerifyConsoleCommandArguments(TEXT("DMX.SetInputPortProtocol"), MinNumExpectedArgs, Args))
+			})
+	);
+
+	static FAutoConsoleCommand GDMXSetInputPortPriorityStrategyCommand(
+		TEXT("DMX.SetInputPortPriorityStrategy"),
+		TEXT("DMX.SetInputPortPriorityStrategy [PortName][PriorityStrategy (0 = None, 1 = Lowest, 2 = LowerThan, 3 = Equal, 4 = Higher Than, 5 = Highest)]. Example: DMX.SetInputPortPriorityStrategy MyInputPort 1"),
+		FConsoleCommandWithArgsDelegate::CreateStatic(
+			[](const TArray<FString>& Args)
 			{
-				return;
-			}
+				constexpr int32 MinNumExpectedArgs = 2;
+				if (!VerifyConsoleCommandArguments(TEXT("DMX.SetInputPortPriorityStrategy"), MinNumExpectedArgs, Args))
+				{
+					return;
+				}
 
-			const FString& PortName = Args[0];
-			const FString& PriorityStrategyValueString = Args[1];
-			uint8 PriorityStrategyValue;
-			if (LexTryParseString<uint8>(PriorityStrategyValue, *PriorityStrategyValueString))
-			{
-				const EDMXPortPriorityStrategy PriorityStrategyEnumValue = static_cast<EDMXPortPriorityStrategy>(PriorityStrategyValue);
-
-				if (PriorityStrategyEnumValue == EDMXPortPriorityStrategy::None ||
-					PriorityStrategyEnumValue == EDMXPortPriorityStrategy::Lowest ||
-					PriorityStrategyEnumValue == EDMXPortPriorityStrategy::LowerThan ||
-					PriorityStrategyEnumValue == EDMXPortPriorityStrategy::Equal ||
-					PriorityStrategyEnumValue == EDMXPortPriorityStrategy::HigherThan ||
-					PriorityStrategyEnumValue == EDMXPortPriorityStrategy::Highest)
+				const FString& PortName = Args[0];
+				const FString& PriorityStrategyValueString = Args[1];
+				uint8 PriorityStrategyValue;
+				if (LexTryParseString<uint8>(PriorityStrategyValue, *PriorityStrategyValueString))
 				{
-					DMX_OVERRIDE_INPUTPORT_VAR(PriorityStrategy, PortName, PriorityStrategyEnumValue);
+					const EDMXPortPriorityStrategy PriorityStrategyEnumValue = static_cast<EDMXPortPriorityStrategy>(PriorityStrategyValue);
+
+					if (PriorityStrategyEnumValue == EDMXPortPriorityStrategy::None ||
+						PriorityStrategyEnumValue == EDMXPortPriorityStrategy::Lowest ||
+						PriorityStrategyEnumValue == EDMXPortPriorityStrategy::LowerThan ||
+						PriorityStrategyEnumValue == EDMXPortPriorityStrategy::Equal ||
+						PriorityStrategyEnumValue == EDMXPortPriorityStrategy::HigherThan ||
+						PriorityStrategyEnumValue == EDMXPortPriorityStrategy::Highest)
+					{
+						DMX_OVERRIDE_INPUTPORT_VAR(PriorityStrategy, PortName, PriorityStrategyEnumValue);
+					}
 				}
-			}
-		})
-);
-
-static FAutoConsoleCommand GDMXSetInputPortPriorityCommand(
-	TEXT("DMX.SetInputPortPriority"),
-	TEXT("DMX.SetInputPortPriority [PortName][Priority]. Sets the priority of the input port. Example: DMX.SetInputPortPriority MyInputPort 100"),
-	FConsoleCommandWithArgsDelegate::CreateStatic(
-		[](const TArray<FString>& Args)
-		{
-			constexpr int32 MinNumExpectedArgs = 2;
-			if (!UE::DMX::DMXInputPort::VerifyConsoleCommandArguments(TEXT("DMX.SetInputPortProtocol"), MinNumExpectedArgs, Args))
+			})
+	);
+
+	static FAutoConsoleCommand GDMXSetInputPortPriorityCommand(
+		TEXT("DMX.SetInputPortPriority"),
+		TEXT("DMX.SetInputPortPriority [PortName][Priority]. Sets the priority of the input port. Example: DMX.SetInputPortPriority MyInputPort 100"),
+		FConsoleCommandWithArgsDelegate::CreateStatic(
+			[](const TArray<FString>& Args)
 			{
-				return;
-			}
+				constexpr int32 MinNumExpectedArgs = 2;
+				if (!VerifyConsoleCommandArguments(TEXT("DMX.SetInputPortPriority"), MinNumExpectedArgs, Args))
+				{
+					return;
+				}
 
-			const FString& PortName = Args[0];
-			const FString& PriorityValueString = Args[1];
-			int32 PriorityValue;
-			if (LexTryParseString<int32>(PriorityValue, *PriorityValueString))
-			{
-				DMX_OVERRIDE_INPUTPORT_VAR(Priority, PortName, PriorityValue);
-			}
-		})
-);
-
-static FAutoConsoleCommand GDMXResetInputPortToProjectSettings(
-	TEXT("DMX.ResetInputPortToProjectSettings"),
-	TEXT("DMX.ResetInputPortToProjectSettings [PortName]. Resets the input port to how it is defined in project settings. Example: DMX.ResetInputPortToProjectSettings MyInputPort"),
-	FConsoleCommandWithArgsDelegate::CreateStatic(
-		[](const TArray<FString>& Args)
-		{
-			constexpr int32 MinNumExpectedArgs = 1;
-			if (!UE::DMX::DMXInputPort::VerifyConsoleCommandArguments(TEXT("DMX.SetInputPortProtocol"), MinNumExpectedArgs, Args))
+				const FString& PortName = Args[0];
+				const FString& PriorityValueString = Args[1];
+				int32 PriorityValue;
+				if (LexTryParseString<int32>(PriorityValue, *PriorityValueString))
+				{
+					DMX_OVERRIDE_INPUTPORT_VAR(Priority, PortName, PriorityValue);
+				}
+			})
+	);
+
+	static FAutoConsoleCommand GDMXResetInputPortToProjectSettings(
+		TEXT("DMX.ResetInputPortToProjectSettings"),
+		TEXT("DMX.ResetInputPortToProjectSettings [PortName]. Resets the input port to how it is defined in project settings. Example: DMX.ResetInputPortToProjectSettings MyInputPort"),
+		FConsoleCommandWithArgsDelegate::CreateStatic(
+			[](const TArray<FString>& Args)
 			{
-				return;
-			}
+				constexpr int32 MinNumExpectedArgs = 1;
+				if (!VerifyConsoleCommandArguments(TEXT("DMX.ResetInputPortToProjectSettings"), MinNumExpectedArgs, Args))
+				{
+					return;
+				}
 
-			const FString& PortName = Args[0];
+				const FString& PortName = Args[0];
 
-			const FDMXInputPortSharedPtr InputPort = UE::DMX::DMXInputPort::FindInputPortByName(PortName);
-			if(InputPort.IsValid())
-			{
-				const UDMXProtocolSettings* ProtocolSettings = GetDefault<UDMXProtocolSettings>();
-				if (ProtocolSettings)
+				const FDMXInputPortSharedPtr InputPort = FindInputPortByName(PortName);
+				if(InputPort.IsValid())
 				{
-					const FDMXInputPortConfig* PortConfigPtr = ProtocolSettings->InputPortConfigs.FindByPredicate([InputPort](const FDMXInputPortConfig& InputPortConfig)
-						{
-							return InputPortConfig.GetPortGuid() == InputPort->GetPortGuid();
-						});
-
-					if (PortConfigPtr)
+					const UDMXProtocolSettings* ProtocolSettings = GetDefault<UDMXProtocolSettings>();
+					if (ProtocolSettings)
 					{
-						FDMXInputPortConfig PortConfig = *PortConfigPtr;
-						InputPort->UpdateFromConfig(PortConfig);
+						const FDMXInputPortConfig* PortConfigPtr = ProtocolSettings->InputPortConfigs.FindByPredicate([InputPort](const FDMXInputPortConfig& InputPortConfig)
+							{
+								return InputPortConfig.GetPortGuid() == InputPort->GetPortGuid();
+							});
+
+						if (PortConfigPtr)
+						{
+							FDMXInputPortConfig PortConfig = *PortConfigPtr;
+							InputPort->UpdateFromConfig(PortConfig);
+						}
 					}
 				}
-			}
-		})
-);
-
+			})
+	);
+}
 #undef DMX_OVERRIDE_INPUTPORT_VAR
 
-
 FDMXInputPortSharedRef FDMXInputPort::CreateFromConfig(FDMXInputPortConfig& InputPortConfig)
 {
 	// Port Configs are expected to have a valid guid always
@@ -391,7 +438,7 @@ FDMXInputPortSharedRef FDMXInputPort::CreateFromConfig(FDMXInputPortConfig& Inpu
 
 FDMXInputPort::~FDMXInputPort()
 {
-	// All Inputs need to be explicitly removed before destruction 
+	// All Listeners need to be explicitly removed before destruction 
 	check(RawListeners.Num() == 0);
 	check(LocalUniverseToListenerGroupMap.Num() == 0);
 
diff --git a/Engine/Plugins/VirtualProduction/DMX/DMXProtocol/Source/DMXProtocol/Private/IO/DMXOutputPort.cpp b/Engine/Plugins/VirtualProduction/DMX/DMXProtocol/Source/DMXProtocol/Private/IO/DMXOutputPort.cpp
--- a/Engine/Plugins/VirtualProduction/DMX/DMXProtocol/Source/DMXProtocol/Private/IO/DMXOutputPort.cpp
+++ b/Engine/Plugins/VirtualProduction/DMX/DMXProtocol/Source/DMXProtocol/Private/IO/DMXOutputPort.cpp
@@ -19,8 +19,60 @@
 #define LOCTEXT_NAMESPACE "DMXOutputPort"
 
 
-namespace UE::DMX::DMXOutputPort
+
+/** Helper to override a member variable of an Output Port */
+#define DMX_OVERRIDE_OUTPUTPORT_VAR(MemberName, PortName, Value) \
+{ \
+	const FDMXOutputPortSharedPtr OverRideOutputPortMacro_OutputPort = UE::DMX::DMXOutputPort::Private::ConsoleCommands::FindOutputPortByName(PortName); \
+	if (OverRideOutputPortMacro_OutputPort.IsValid()) \
+	{ \
+		const FDMXOutputPortConfig OldOutputPortConfig = OverRideOutputPortMacro_OutputPort->MakeOutputPortConfig(); \
+		FDMXOutputPortConfigParams OutputPortConfigParams = FDMXOutputPortConfigParams(OldOutputPortConfig); \
+		OutputPortConfigParams.MemberName = Value; \
+		FDMXOutputPortConfig NewOutputPortConfig(OverRideOutputPortMacro_OutputPort->GetPortGuid(), OutputPortConfigParams); \
+		constexpr bool bForceUpdateRegistrationWithProtocol = true; \
+		OverRideOutputPortMacro_OutputPort->UpdateFromConfig(NewOutputPortConfig, bForceUpdateRegistrationWithProtocol); \
+	} \
+}
+
+namespace UE::DMX::DMXOutputPort::Private::ConsoleCommands
 {
+	/** Logs the configuration of the Output Port */
+	void LogOutputPortConfiguration(const FDMXOutputPortSharedPtr OutputPort)
+	{
+		if (OutputPort.IsValid())
+		{
+			UE_LOG(LogDMXProtocol, Log, TEXT("================================================"));
+			UE_LOG(LogDMXProtocol, Log, TEXT("Configuration of Output Port '%s':"), *OutputPort->GetPortName());
+			UE_LOG(LogDMXProtocol, Log, TEXT(" "));
+
+			UE_LOG(LogDMXProtocol, Log, TEXT("Port Name: %s"), *OutputPort->GetPortName());
+
+			FDMXOutputPortConfig PortConfig = OutputPort->MakeOutputPortConfig();
+			UE_LOG(LogDMXProtocol, Log, TEXT("Protocol: %s"), *PortConfig.GetProtocolName().ToString());
+			UE_LOG(LogDMXProtocol, Log, TEXT("Communication Type: %i"), static_cast<int32>(PortConfig.GetCommunicationType()));
+			UE_LOG(LogDMXProtocol, Log, TEXT("Device Address: %s"), *PortConfig.GetDeviceAddress());
+
+			const TArray<FDMXOutputPortDestinationAddress> DestinationAddresses = PortConfig.GetDestinationAddresses();
+			for (int32 IndexDestinationAddress = 0; IndexDestinationAddress < DestinationAddresses.Num(); IndexDestinationAddress++)
+			{
+				UE_LOG(LogDMXProtocol, Log, TEXT("Destination Address [%i]: %s"), IndexDestinationAddress, *DestinationAddresses[IndexDestinationAddress].DestinationAddressString);
+			}
+
+			UE_LOG(LogDMXProtocol, Log, TEXT("Input into Engine: %s"), PortConfig.NeedsLoopbackToEngine() ? TEXT("True") : TEXT("False"));
+			UE_LOG(LogDMXProtocol, Log, TEXT("Local Universe Start: %i"), PortConfig.GetLocalUniverseStart());
+			UE_LOG(LogDMXProtocol, Log, TEXT("Num Universes: %i"), PortConfig.GetNumUniverses());
+			UE_LOG(LogDMXProtocol, Log, TEXT("Extern Universe Start: %i"), PortConfig.GetExternUniverseStart());
+			UE_LOG(LogDMXProtocol, Log, TEXT("Priority:	%i"), PortConfig.GetPriority());
+
+			UE_LOG(LogDMXProtocol, Log, TEXT("Delay: %f"), PortConfig.GetDelay());
+			UE_LOG(LogDMXProtocol, Log, TEXT("Delay Frame Rate:	%i/%i"), PortConfig.GetDelayFrameRate().Numerator, PortConfig.GetDelayFrameRate().Denominator);
+
+			UE_LOG(LogDMXProtocol, Log, TEXT("Delay in Seconds: %f"), OutputPort->GetDelaySeconds());
+			UE_LOG(LogDMXProtocol, Log, TEXT("================================================"));
+		}
+	}
+
 	/** Finds an Output port by its name, optionally logging when it can't be found */
 	static const FDMXOutputPortSharedPtr FindOutputPortByName(const FString& PortName, bool bPrintToLogIfNotFound = true)
 	{
@@ -35,7 +87,14 @@ namespace UE::DMX::DMXOutputPort
 		}
 		else if (bPrintToLogIfNotFound)
 		{
-			UE_LOG(LogDMXProtocol, Warning, TEXT("Could not find Output Port '%s'.  Available Output Ports and their names can be found in Project settings -> Plugins -> DMX."), *PortName);
+			if (PortName.IsEmpty())
+			{
+				UE_LOG(LogDMXProtocol, Warning, TEXT("No Ouptut Port specified for console command. First argument needs to be a valid Port Name."));
+			}
+			else
+			{
+				UE_LOG(LogDMXProtocol, Warning, TEXT("Could not find Output Port '%s'.  Available Output Ports and their names can be found in Project settings -> Plugins -> DMX."), *PortName);
+			}
 		}
 
 		return nullptr;
@@ -44,442 +103,428 @@ namespace UE::DMX::DMXOutputPort
 	/** Tests the console command, returns true if they yield a valid port. Logs issues otherwise. Handles ? commands. */
 	static bool VerifyConsoleCommandArguments(const FString& ConsoleCommand, int32 MinNumExpectedArgs, const TArray<FString>& Args)
 	{
-		check(MinNumExpectedArgs >= 1);
+		const FString PortName = Args.Num() > 0 ? Args[0] : FString();
+		const FDMXOutputPortSharedPtr OutputPort = FindOutputPortByName(PortName);
+
+		const bool bOnlyPortSpecified = Args.Num() == 1;
+		const bool bSecondArgIsQuestionmark = (Args.Num() == 2 && Args[1] == TEXT("?"));
+		if (OutputPort.IsValid() && (bOnlyPortSpecified || bSecondArgIsQuestionmark))
+		{		
+			// Log the port if the there is only one arg or the second arg is '?' (e.g.'DMX.SetOutputPortProtocol OutputPortA' or 'DMX.SetOutputPortProtocol OutputPortA ?')
+			LogOutputPortConfiguration(OutputPort);
 
-		if (Args.Num() < MinNumExpectedArgs)
-		{
-			UE_LOG(LogDMXProtocol, Warning, TEXT("Console Command %s failed. Insufficient arguments specified."), *ConsoleCommand);
-			UE_LOG(LogDMXProtocol, Warning, TEXT("Use '%s ?' for help. Use '%s PortName ?' to print the current configuration to logs."), *ConsoleCommand, *ConsoleCommand);
 			return false;
 		}
+		else if (!OutputPort.IsValid())
+		{			
+			// Log additional help if there were not enough arguments or arguments were invalid
+			UE_LOG(LogDMXProtocol, Warning, TEXT("Console Command %s failed. Output Port with Name '%s' does not exist."), *ConsoleCommand, *PortName);
+			UE_LOG(LogDMXProtocol, Warning, TEXT("Ports and their Names can be found in Project Settings -> Plugins -> DMX."));
 
-		const FString& PortName = Args[0];
-		const FDMXOutputPortSharedPtr OutputPort = UE::DMX::DMXOutputPort::FindOutputPortByName(PortName);
-		if (OutputPort.IsValid())
+			return false;
+		}
+		else if (MinNumExpectedArgs > Args.Num())
 		{
-			// Log the port if the console command sets the second arg to '?' (e.g. 'DMX.SetOutputPortProtocol OutputPortA ?')
-			if (Args.Num() == 2 && Args[1] == TEXT("?"))
-			{
-				// Log the port
-				UE_LOG(LogDMXProtocol, Log, TEXT("Output Port: %s"), *OutputPort->GetPortName());
-
-				FDMXOutputPortConfig PortConfig = OutputPort->MakeOutputPortConfig();
-				UE_LOG(LogDMXProtocol, Log, TEXT("Protocol: %s"), *PortConfig.GetProtocolName().ToString());
-				UE_LOG(LogDMXProtocol, Log, TEXT("Communication Type: %i"), static_cast<int32>(PortConfig.GetCommunicationType()));
-				UE_LOG(LogDMXProtocol, Log, TEXT("Device Address: %s"), *PortConfig.GetDeviceAddress());
-
-				const TArray<FDMXOutputPortDestinationAddress> DestinationAddresses = PortConfig.GetDestinationAddresses();
-				for (int32 IndexDestinationAddress = 0; IndexDestinationAddress < DestinationAddresses.Num(); IndexDestinationAddress++)
-				{
-					UE_LOG(LogDMXProtocol, Log, TEXT("Destination Address [%i]: %s"), IndexDestinationAddress, *DestinationAddresses[IndexDestinationAddress].DestinationAddressString);
-				}
+			UE_LOG(LogDMXProtocol, Warning, TEXT("Console Command %s failed. Invalid arguments specified."), *ConsoleCommand);
 
-				UE_LOG(LogDMXProtocol, Log, TEXT("Input into Engine: %s"), PortConfig.NeedsLoopbackToEngine() ? TEXT("True") : TEXT("False"));
-				UE_LOG(LogDMXProtocol, Log, TEXT("Local Universe Start: %i"), PortConfig.GetLocalUniverseStart());
-				UE_LOG(LogDMXProtocol, Log, TEXT("Num Universes: %i"), PortConfig.GetNumUniverses());
-				UE_LOG(LogDMXProtocol, Log, TEXT("Extern Universe Start: %i"), PortConfig.GetExternUniverseStart());
-				UE_LOG(LogDMXProtocol, Log, TEXT("Priority:	%i"), PortConfig.GetPriority());
-
-				UE_LOG(LogDMXProtocol, Log, TEXT("Delay: %f"), PortConfig.GetDelay());
-				UE_LOG(LogDMXProtocol, Log, TEXT("Delay Frame Rate:	%i/%i"), PortConfig.GetDelayFrameRate().Numerator, PortConfig.GetDelayFrameRate().Denominator);
-				UE_LOG(LogDMXProtocol, Log, TEXT("Delay in Seconds: %f"), OutputPort->GetDelaySeconds());
-			}
-			else
-			{
-				return true;
-			}
+			return false;
 		}
 
-		return false;
+		return true;
 	}
-}
 
-/** Helper to override a member variable of an Output Port */
-#define DMX_OVERRIDE_OUTPUTPORT_VAR(MemberName, PortName, Value) \
-{ \
-	const FDMXOutputPortSharedPtr OverRideOutputPortMacro_OutputPort = UE::DMX::DMXOutputPort::FindOutputPortByName(PortName); \
-	if (OverRideOutputPortMacro_OutputPort.IsValid()) \
-	{ \
-		const FDMXOutputPortConfig OldOutputPortConfig = OverRideOutputPortMacro_OutputPort->MakeOutputPortConfig(); \
-		FDMXOutputPortConfigParams OutputPortConfigParams = FDMXOutputPortConfigParams(OldOutputPortConfig); \
-		OutputPortConfigParams.MemberName = Value; \
-		FDMXOutputPortConfig NewOutputPortConfig(OverRideOutputPortMacro_OutputPort->GetPortGuid(), OutputPortConfigParams); \
-		constexpr bool bForceUpdateRegistrationWithProtocol = true; \
-		OverRideOutputPortMacro_OutputPort->UpdateFromConfig(NewOutputPortConfig, bForceUpdateRegistrationWithProtocol); \
-	} \
-}
-
-static FAutoConsoleCommand GDMXSetOutputPortProtocolCommand(
-	TEXT("DMX.SetOutputPortProtocol"),
-	TEXT("DMX.SetOutputPortProtocol [PortName][ProtocolName]. Sets the protocol used by the output port. Example: DMX.SetOutputPortProtocol MyOutputPort Art-Net"),
-	FConsoleCommandWithArgsDelegate::CreateStatic(
-		[](const TArray<FString>& Args)
-		{
-			constexpr int32 MinNumExpectedArgs = 2;
-			if (!UE::DMX::DMXOutputPort::VerifyConsoleCommandArguments(TEXT("DMX.SetOutputPortProtocol"), MinNumExpectedArgs, Args))
-			{
-				return;
-			}
-
-			const FString& PortName = Args[0];
-			const FName ProtocolNameValue = FName(Args[1]);
-
-			if (IDMXProtocol::GetProtocolNames().Contains(ProtocolNameValue))
-			{
-				DMX_OVERRIDE_OUTPUTPORT_VAR(ProtocolName, PortName, ProtocolNameValue);
-			}
-			else
+	static FAutoConsoleCommand GDMXLogAllOutputPortConfigurationsCommand(
+		TEXT("DMX.LogAllOutputPortConfigurations"),
+		TEXT("Logs all Output Port configurations"),
+		FConsoleCommandDelegate::CreateStatic(
+			[]()
 			{
-				UE_LOG(LogDMXProtocol, Warning, TEXT("Console command DMX.SetOutputPortProtocol failed. Protocol '%s' does not exist or is not loaded. Available protocols are: "), *ProtocolNameValue.ToString());
-				for (const FName& ProtocolName : IDMXProtocol::GetProtocolNames())
+				for (const FDMXOutputPortSharedRef& OutputPort : FDMXPortManager::Get().GetOutputPorts())
 				{
-					UE_LOG(LogDMXProtocol, Warning, TEXT("%s"), *ProtocolName.ToString());
+					LogOutputPortConfiguration(OutputPort);
 				}
-			}
-		})
-);
-
-static FAutoConsoleCommand GDMXSetOutputPortCommunicationTypeCommand(
-	TEXT("DMX.SetOutputPortCommunicationType"),
-	TEXT("DMX.SetOutputPortCommunicationType [PortName][CommunicationType (0 = Broadcast, 1 = Unicast, 2 = Multicast)]. Sets the communication type of an output port. Example: DMX.SetOutputPortCommunicationType MyOutputPort 2"),
-	FConsoleCommandWithArgsDelegate::CreateStatic(
-		[](const TArray<FString>& Args)
-		{
-			constexpr int32 MinNumExpectedArgs = 2;
-			if (!UE::DMX::DMXOutputPort::VerifyConsoleCommandArguments(TEXT("DMX.SetOutputPortCommunicationType"), MinNumExpectedArgs, Args))
+			})
+	);
+
+	static FAutoConsoleCommand GDMXSetOutputPortProtocolCommand(
+		TEXT("DMX.SetOutputPortProtocol"),
+		TEXT("DMX.SetOutputPortProtocol [PortName][ProtocolName]. Sets the protocol used by the output port. Example: DMX.SetOutputPortProtocol MyOutputPort Art-Net"),
+		FConsoleCommandWithArgsDelegate::CreateStatic(
+			[](const TArray<FString>& Args)
 			{
-				return;
-			}
+				constexpr int32 MinNumExpectedArgs = 2;
+				if (!VerifyConsoleCommandArguments(TEXT("DMX.SetOutputPortProtocol"), MinNumExpectedArgs, Args))
+				{
+					return;
+				}
 
-			const FString& PortName = Args[0];
-			const FString& CommunicationTypeValueString = Args[1];
+				const FString& PortName = Args[0];
+				const FName ProtocolNameValue = FName(Args[1]);
 
-			uint8 CommunicationTypeValue;
-			if (LexTryParseString<uint8>(CommunicationTypeValue, *CommunicationTypeValueString))
+				if (IDMXProtocol::GetProtocolNames().Contains(ProtocolNameValue))
+				{
+					DMX_OVERRIDE_OUTPUTPORT_VAR(ProtocolName, PortName, ProtocolNameValue);
+				}
+				else
+				{
+					UE_LOG(LogDMXProtocol, Warning, TEXT("Console command DMX.SetOutputPortProtocol failed. Protocol '%s' does not exist or is not loaded. Available protocols are: "), *ProtocolNameValue.ToString());
+					for (const FName& ProtocolName : IDMXProtocol::GetProtocolNames())
+					{
+						UE_LOG(LogDMXProtocol, Warning, TEXT("%s"), *ProtocolName.ToString());
+					}
+				}
+			})
+	);
+
+	static FAutoConsoleCommand GDMXSetOutputPortCommunicationTypeCommand(
+		TEXT("DMX.SetOutputPortCommunicationType"),
+		TEXT("DMX.SetOutputPortCommunicationType [PortName][CommunicationType (0 = Broadcast, 1 = Unicast, 2 = Multicast)]. Sets the communication type of an output port. Example: DMX.SetOutputPortCommunicationType MyOutputPort 2"),
+		FConsoleCommandWithArgsDelegate::CreateStatic(
+			[](const TArray<FString>& Args)
 			{
-				const EDMXCommunicationType CommunicationTypeEnumValue = static_cast<EDMXCommunicationType>(CommunicationTypeValue);
-				if (CommunicationTypeEnumValue == EDMXCommunicationType::Broadcast ||
-					CommunicationTypeEnumValue == EDMXCommunicationType::Multicast ||
-					CommunicationTypeEnumValue == EDMXCommunicationType::Unicast)
+				constexpr int32 MinNumExpectedArgs = 2;
+				if (!VerifyConsoleCommandArguments(TEXT("DMX.SetOutputPortCommunicationType"), MinNumExpectedArgs, Args))
 				{
-					const FDMXOutputPortSharedPtr OutputPort = UE::DMX::DMXOutputPort::FindOutputPortByName(PortName);
+					return;
+				}
+
+				const FString& PortName = Args[0];
+				const FString& CommunicationTypeValueString = Args[1];
 
-					if (OutputPort.IsValid())
+				uint8 CommunicationTypeValue;
+				if (LexTryParseString<uint8>(CommunicationTypeValue, *CommunicationTypeValueString))
+				{
+					const EDMXCommunicationType CommunicationTypeEnumValue = static_cast<EDMXCommunicationType>(CommunicationTypeValue);
+					if (CommunicationTypeEnumValue == EDMXCommunicationType::Broadcast ||
+						CommunicationTypeEnumValue == EDMXCommunicationType::Multicast ||
+						CommunicationTypeEnumValue == EDMXCommunicationType::Unicast)
 					{
-						if (const IDMXProtocolPtr Protocol = OutputPort->GetProtocol())
+						const FDMXOutputPortSharedPtr OutputPort = FindOutputPortByName(PortName);
+
+						if (OutputPort.IsValid())
 						{
-							if (Protocol->GetOutputPortCommunicationTypes().Contains(CommunicationTypeEnumValue))
+							if (const IDMXProtocolPtr Protocol = OutputPort->GetProtocol())
 							{
-								DMX_OVERRIDE_OUTPUTPORT_VAR(CommunicationType, PortName, CommunicationTypeEnumValue);
+								if (Protocol->GetOutputPortCommunicationTypes().Contains(CommunicationTypeEnumValue))
+								{
+									DMX_OVERRIDE_OUTPUTPORT_VAR(CommunicationType, PortName, CommunicationTypeEnumValue);
+								}
 							}
 						}
 					}
 				}
-			}
-		})
-);
+			})
+	);
 
 
-static FAutoConsoleCommand GDMXSetOutputPortOutputPortDeviceAddressCommand(
-	TEXT("DMX.SetOutputPortDeviceAddress"),
-	TEXT("DMX.SetOutputPortDeviceAddress [PortName][DeviceAddress]. Sets the Device Address of an output port, usually the network interface card IP address. Example: DMX.SetInputPortDeviceAddress MyOutputPort 123.45.67.89"),
-	FConsoleCommandWithArgsDelegate::CreateStatic(
-		[](const TArray<FString>& Args)
-		{
-			constexpr int32 MinNumExpectedArgs = 2;
-			if (!UE::DMX::DMXOutputPort::VerifyConsoleCommandArguments(TEXT("DMX.SetOutputPortDeviceAddress"), MinNumExpectedArgs, Args))
+	static FAutoConsoleCommand GDMXSetOutputPortOutputPortDeviceAddressCommand(
+		TEXT("DMX.SetOutputPortDeviceAddress"),
+		TEXT("DMX.SetOutputPortDeviceAddress [PortName][DeviceAddress]. Sets the Device Address of an output port, usually the network interface card IP address. Example: DMX.SetOutputPortDeviceAddress MyOutputPort 123.45.67.89"),
+		FConsoleCommandWithArgsDelegate::CreateStatic(
+			[](const TArray<FString>& Args)
 			{
-				return;
-			}
+				constexpr int32 MinNumExpectedArgs = 2;
+				if (!VerifyConsoleCommandArguments(TEXT("DMX.SetOutputPortDeviceAddress"), MinNumExpectedArgs, Args))
+				{
+					UE_LOG(LogDMXProtocol, Log, TEXT("================================================"));
+					// For this command, additionally log the available Local Network Interface Card IP addresses
+					UE_LOG(LogDMXProtocol, Log, TEXT("Logging available Local Network Interface Card IP Addresses:"));
+					const TArray<TSharedPtr<FString>> AvailableIPAddresses = FDMXProtocolUtils::GetLocalNetworkInterfaceCardIPs();
+					for (int32 IPAddressIndex = 0; IPAddressIndex < AvailableIPAddresses.Num(); IPAddressIndex++)
+					{
+						UE_LOG(LogDMXProtocol, Log, TEXT("%i: %s"), IPAddressIndex, *(*AvailableIPAddresses[IPAddressIndex]));
+					}
+					UE_LOG(LogDMXProtocol, Log, TEXT("================================================"));
+					return;
+				}
 
-			const FString& PortName = Args[0];
-			const FString& DeviceAddressValue = Args[1];
+				const FString& PortName = Args[0];
+				const FString& DeviceAddressValue = Args[1];
 
-			DMX_OVERRIDE_OUTPUTPORT_VAR(DeviceAddress, PortName, DeviceAddressValue);
-		})
-);
+				DMX_OVERRIDE_OUTPUTPORT_VAR(DeviceAddress, PortName, DeviceAddressValue);
+			})
+	);
 
-static FAutoConsoleCommand GDMXSetOutputPortOutputPortDestinationAddressesCommand(
-	TEXT("DMX.SetOutputPortDestinationAddresses"),
-	TEXT("DMX.SetOutputPortDestinationAddresses [PortName][DestinationAddress1][DestinationAddress2][...][DestinationAddressN]. Sets the Destination Addresses of an output port. Example: DMX.SetInputPortDeviceAddress MyOutputPort 11.33.55.77 22.44.66.88"),
-	FConsoleCommandWithArgsDelegate::CreateStatic(
-		[](const TArray<FString>& Args)
-		{
-			constexpr int32 MinNumExpectedArgs = 2;
-			if (!UE::DMX::DMXOutputPort::VerifyConsoleCommandArguments(TEXT("DMX.SetOutputPortDestinationAddresses"), MinNumExpectedArgs, Args))
+	static FAutoConsoleCommand GDMXSetOutputPortOutputPortDestinationAddressesCommand(
+		TEXT("DMX.SetOutputPortDestinationAddresses"),
+		TEXT("DMX.SetOutputPortDestinationAddresses [PortName][DestinationAddress1][DestinationAddress2][...][DestinationAddressN]. Sets the Destination Addresses of an output port. Example: DMX.SetOutputPortDeviceAddress MyOutputPort 11.33.55.77 22.44.66.88"),
+		FConsoleCommandWithArgsDelegate::CreateStatic(
+			[](const TArray<FString>& Args)
 			{
-				return;
-			}
+				constexpr int32 MinNumExpectedArgs = 2;
+				if (!VerifyConsoleCommandArguments(TEXT("DMX.SetOutputPortDestinationAddresses"), MinNumExpectedArgs, Args))
+				{
+					return;
+				}
 
-			const FString& PortName = Args[0];
-			TArray<FDMXOutputPortDestinationAddress>  DestinationAddressesValue;
-			for (int32 ArgIndex = 1; ArgIndex < Args.Num(); ArgIndex++)
-			{
-				const FDMXOutputPortDestinationAddress Address = FDMXOutputPortDestinationAddress(Args[ArgIndex]);
-				DestinationAddressesValue.Add(Address);
-			}
+				const FString& PortName = Args[0];
+				TArray<FDMXOutputPortDestinationAddress>  DestinationAddressesValue;
+				for (int32 ArgIndex = 1; ArgIndex < Args.Num(); ArgIndex++)
+				{
+					const FDMXOutputPortDestinationAddress Address = FDMXOutputPortDestinationAddress(Args[ArgIndex]);
+					DestinationAddressesValue.Add(Address);
+				}
 
-			DMX_OVERRIDE_OUTPUTPORT_VAR(DestinationAddresses, PortName, DestinationAddressesValue);
-		})
-);
+				DMX_OVERRIDE_OUTPUTPORT_VAR(DestinationAddresses, PortName, DestinationAddressesValue);
+			})
+	);
 
-static FAutoConsoleCommand GDMXSetOutputPortLoopbackToEngineCommand(
-	TEXT("DMX.SetOutputPortInputIntoEngine"),
-	TEXT("DMX.SetOutputPortInputIntoEngine [PortName][Flag]. Sets if the Output Port is input into the engine directly. Example: DMX.SetOutputPortInputIntoEngine MyOutputPort 1"),
-	FConsoleCommandWithArgsDelegate::CreateStatic(
-		[](const TArray<FString>& Args)
-		{
-			constexpr int32 MinNumExpectedArgs = 2;
-			if (!UE::DMX::DMXOutputPort::VerifyConsoleCommandArguments(TEXT("DMX.SetOutputPortInputIntoEngine"), MinNumExpectedArgs, Args))
+	static FAutoConsoleCommand GDMXSetOutputPortLoopbackToEngineCommand(
+		TEXT("DMX.SetOutputPortInputIntoEngine"),
+		TEXT("DMX.SetOutputPortInputIntoEngine [PortName][Flag]. Sets if the Output Port is input into the engine directly. Example: DMX.SetOutputPortInputIntoEngine MyOutputPort 1"),
+		FConsoleCommandWithArgsDelegate::CreateStatic(
+			[](const TArray<FString>& Args)
 			{
-				return;
-			}
+				constexpr int32 MinNumExpectedArgs = 2;
+				if (!VerifyConsoleCommandArguments(TEXT("DMX.SetOutputPortInputIntoEngine"), MinNumExpectedArgs, Args))
+				{
+					return;
+				}
 
-			const FString& PortName = Args[0];
-			const bool bLoopbackToEngine = Args[1] == TEXT("1") || Args[1].Equals(TEXT("true"), ESearchCase::IgnoreCase);
+				const FString& PortName = Args[0];
+				const bool bLoopbackToEngine = Args[1] == TEXT("1") || Args[1].Equals(TEXT("true"), ESearchCase::IgnoreCase);
 
-			DMX_OVERRIDE_OUTPUTPORT_VAR(bLoopbackToEngine, PortName, bLoopbackToEngine);
-		})
-);
+				DMX_OVERRIDE_OUTPUTPORT_VAR(bLoopbackToEngine, PortName, bLoopbackToEngine);
+			})
+	);
 
-static FAutoConsoleCommand GDMXSetOutputPortLocalUniverseStartCommand(
-	TEXT("DMX.SetOutputPortLocalUniverseStart"),
-	TEXT("DMX.SetOutputPortLocalUniverseStart [PortName][Universe]. Sets the local universe start of the output port. Example: DMX.SetOutputPortLocalUniverseStart MyOutputPort 5"),
-	FConsoleCommandWithArgsDelegate::CreateStatic(
-		[](const TArray<FString>& Args)
-		{
-			constexpr int32 MinNumExpectedArgs = 2;
-			if (!UE::DMX::DMXOutputPort::VerifyConsoleCommandArguments(TEXT("DMX.SetOutputPortLocalUniverseStart"), MinNumExpectedArgs, Args))
+	static FAutoConsoleCommand GDMXSetOutputPortLocalUniverseStartCommand(
+		TEXT("DMX.SetOutputPortLocalUniverseStart"),
+		TEXT("DMX.SetOutputPortLocalUniverseStart [PortName][Universe]. Sets the local universe start of the output port. Example: DMX.SetOutputPortLocalUniverseStart MyOutputPort 5"),
+		FConsoleCommandWithArgsDelegate::CreateStatic(
+			[](const TArray<FString>& Args)
 			{
-				return;
-			}
+				constexpr int32 MinNumExpectedArgs = 2;
+				if (!VerifyConsoleCommandArguments(TEXT("DMX.SetOutputPortLocalUniverseStart"), MinNumExpectedArgs, Args))
+				{
+					return;
+				}
 
-			const FString& PortName = Args[0];
-			const FDMXOutputPortSharedPtr OutputPort = UE::DMX::DMXOutputPort::FindOutputPortByName(PortName);
+				const FString& PortName = Args[0];
+				const FDMXOutputPortSharedPtr OutputPort = FindOutputPortByName(PortName);
 
-			if (OutputPort.IsValid())
-			{
-				if (const IDMXProtocolPtr& Protocol = OutputPort->GetProtocol())
+				if (OutputPort.IsValid())
 				{
-					const FString& LocalUniverseStartValueString = Args[1];
-					int32 LocalUniverseStartValue;
-					if (LexTryParseString<int32>(LocalUniverseStartValue, *LocalUniverseStartValueString))
+					if (const IDMXProtocolPtr& Protocol = OutputPort->GetProtocol())
 					{
-						const int32 NumUniverses = OutputPort->GetLocalUniverseEnd() - OutputPort->GetLocalUniverseStart() + 1;
+						const FString& LocalUniverseStartValueString = Args[1];
+						int32 LocalUniverseStartValue;
+						if (LexTryParseString<int32>(LocalUniverseStartValue, *LocalUniverseStartValueString))
+						{
+							const int32 NumUniverses = OutputPort->GetLocalUniverseEnd() - OutputPort->GetLocalUniverseStart() + 1;
 
-						const int32 ExternUniverseStart = OutputPort->ConvertLocalToExternUniverseID(LocalUniverseStartValue);
-						const int32 ExternUniverseEnd = OutputPort->ConvertLocalToExternUniverseID(LocalUniverseStartValue + NumUniverses - 1);
+							const int32 ExternUniverseStart = OutputPort->ConvertLocalToExternUniverseID(LocalUniverseStartValue);
+							const int32 ExternUniverseEnd = OutputPort->ConvertLocalToExternUniverseID(LocalUniverseStartValue + NumUniverses - 1);
 
-						if (Protocol->IsValidUniverseID(ExternUniverseStart) &&
-							Protocol->IsValidUniverseID(ExternUniverseEnd))
-						{
-							DMX_OVERRIDE_OUTPUTPORT_VAR(LocalUniverseStart, PortName, LocalUniverseStartValue);
-						}
-						else
-						{
-							UE_LOG(LogDMXProtocol, Warning, TEXT("Console command DMX.SetOutputPortLocalUniverseStart failed. Local Universe Start '%s' along with Num Universes '%i' results in a Universe range that is not supported by the Protocol of the Port."), *LocalUniverseStartValueString, NumUniverses);
+							if (Protocol->IsValidUniverseID(ExternUniverseStart) &&
+								Protocol->IsValidUniverseID(ExternUniverseEnd))
+							{
+								DMX_OVERRIDE_OUTPUTPORT_VAR(LocalUniverseStart, PortName, LocalUniverseStartValue);
+							}
+							else
+							{
+								UE_LOG(LogDMXProtocol, Warning, TEXT("Console command DMX.SetOutputPortLocalUniverseStart failed. Local Universe Start '%s' along with Num Universes '%i' results in a Universe range that is not supported by the Protocol of the Port."), *LocalUniverseStartValueString, NumUniverses);
+							}
 						}
 					}
 				}
-			}
-		})
-);
-
-static FAutoConsoleCommand GDMXSetOutputPortNumUniversesCommand(
-	TEXT("DMX.SetOutputPortNumUniverses"),
-	TEXT("DMX.SetOutputPortNumUniverses [PortName][Universe]. Sets the num universes of the output port. Example: DMX.SetOutputPortNumUniverses MyOutputPort 10"),
-	FConsoleCommandWithArgsDelegate::CreateStatic(
-		[](const TArray<FString>& Args)
-		{
-			constexpr int32 MinNumExpectedArgs = 2;
-			if (!UE::DMX::DMXOutputPort::VerifyConsoleCommandArguments(TEXT("DMX.SetOutputPortNumUniverses"), MinNumExpectedArgs, Args))
+			})
+	);
+
+	static FAutoConsoleCommand GDMXSetOutputPortNumUniversesCommand(
+		TEXT("DMX.SetOutputPortNumUniverses"),
+		TEXT("DMX.SetOutputPortNumUniverses [PortName][Universe]. Sets the num universes of the output port. Example: DMX.SetOutputPortNumUniverses MyOutputPort 10"),
+		FConsoleCommandWithArgsDelegate::CreateStatic(
+			[](const TArray<FString>& Args)
 			{
-				return;
-			}
+				constexpr int32 MinNumExpectedArgs = 2;
+				if (!VerifyConsoleCommandArguments(TEXT("DMX.SetOutputPortNumUniverses"), MinNumExpectedArgs, Args))
+				{
+					return;
+				}
 
-			const FString& PortName = Args[0];
-			const FDMXOutputPortSharedPtr OutputPort = UE::DMX::DMXOutputPort::FindOutputPortByName(PortName);
+				const FString& PortName = Args[0];
+				const FDMXOutputPortSharedPtr OutputPort = FindOutputPortByName(PortName);
 
-			if (OutputPort.IsValid())
-			{
-				if (const IDMXProtocolPtr& Protocol = OutputPort->GetProtocol())
+				if (OutputPort.IsValid())
 				{
-					const FString& NumUniversesValueString = Args[1];
-					int32 NumUniversesValue;
-					if (LexTryParseString<int32>(NumUniversesValue, *NumUniversesValueString))
+					if (const IDMXProtocolPtr& Protocol = OutputPort->GetProtocol())
 					{
-						const int32 LocalUniverseStart = OutputPort->GetLocalUniverseStart();
+						const FString& NumUniversesValueString = Args[1];
+						int32 NumUniversesValue;
+						if (LexTryParseString<int32>(NumUniversesValue, *NumUniversesValueString))
+						{
+							const int32 LocalUniverseStart = OutputPort->GetLocalUniverseStart();
 
-						const int32 ExternUniverseStart = OutputPort->ConvertLocalToExternUniverseID(LocalUniverseStart);
-						const int32 ExternUniverseEnd = ExternUniverseStart + NumUniversesValue - 1;
+							const int32 ExternUniverseStart = OutputPort->ConvertLocalToExternUniverseID(LocalUniverseStart);
+							const int32 ExternUniverseEnd = ExternUniverseStart + NumUniversesValue - 1;
 
-						if (NumUniversesValue > 0 &&
-							Protocol->IsValidUniverseID(ExternUniverseStart) &&
-							Protocol->IsValidUniverseID(ExternUniverseEnd))
-						{
-							DMX_OVERRIDE_OUTPUTPORT_VAR(NumUniverses, PortName, NumUniversesValue);
-						}
-						else
-						{
-							UE_LOG(LogDMXProtocol, Warning, TEXT("Console command DMX.SetOutputPortNumUniverses failed. Local Universe Start '%i' along with Num Universes '%s' results in a Universe range that is not supported by the Protocol of the Port."), LocalUniverseStart, *NumUniversesValueString);
+							if (NumUniversesValue > 0 &&
+								Protocol->IsValidUniverseID(ExternUniverseStart) &&
+								Protocol->IsValidUniverseID(ExternUniverseEnd))
+							{
+								DMX_OVERRIDE_OUTPUTPORT_VAR(NumUniverses, PortName, NumUniversesValue);
+							}
+							else
+							{
+								UE_LOG(LogDMXProtocol, Warning, TEXT("Console command DMX.SetOutputPortNumUniverses failed. Local Universe Start '%i' along with Num Universes '%s' results in a Universe range that is not supported by the Protocol of the Port."), LocalUniverseStart, *NumUniversesValueString);
+							}
 						}
 					}
 				}
-			}
-		})
-);
-
-static FAutoConsoleCommand GDMXSetOutputPortExternUniverseStartCommand(
-	TEXT("DMX.SetOutputPortExternUniverseStart"),
-	TEXT("DMX.SetOutputPortExternUniverseStart [PortName][Universe]. Sets the extern universe start of the output port. Example: DMX.SetOutputPortExternUniverseStart MyOutputPort 7"),
-	FConsoleCommandWithArgsDelegate::CreateStatic(
-		[](const TArray<FString>& Args)
-		{
-			constexpr int32 MinNumExpectedArgs = 2;
-			if (!UE::DMX::DMXOutputPort::VerifyConsoleCommandArguments(TEXT("DMX.SetOutputPortExternUniverseStart"), MinNumExpectedArgs, Args))
+			})
+	);
+
+	static FAutoConsoleCommand GDMXSetOutputPortExternUniverseStartCommand(
+		TEXT("DMX.SetOutputPortExternUniverseStart"),
+		TEXT("DMX.SetOutputPortExternUniverseStart [PortName][Universe]. Sets the extern universe start of the output port. Example: DMX.SetOutputPortExternUniverseStart MyOutputPort 7"),
+		FConsoleCommandWithArgsDelegate::CreateStatic(
+			[](const TArray<FString>& Args)
 			{
-				return;
-			}
+				constexpr int32 MinNumExpectedArgs = 2;
+				if (!VerifyConsoleCommandArguments(TEXT("DMX.SetOutputPortExternUniverseStart"), MinNumExpectedArgs, Args))
+				{
+					return;
+				}
 
-			const FString& PortName = Args[0];
-			const FDMXOutputPortSharedPtr OutputPort = UE::DMX::DMXOutputPort::FindOutputPortByName(PortName);
+				const FString& PortName = Args[0];
+				const FDMXOutputPortSharedPtr OutputPort = FindOutputPortByName(PortName);
 
 
-			if (OutputPort.IsValid())
-			{
-				if (const IDMXProtocolPtr& Protocol = OutputPort->GetProtocol())
+				if (OutputPort.IsValid())
 				{
-					const FString& ExternUniverseStartValueString = Args[1];
-					int32 ExternUniverseStartValue;
-					if (LexTryParseString<int32>(ExternUniverseStartValue, *ExternUniverseStartValueString))
+					if (const IDMXProtocolPtr& Protocol = OutputPort->GetProtocol())
 					{
-						const int32 NumUniverses = OutputPort->GetLocalUniverseEnd() - OutputPort->GetLocalUniverseStart() + 1;
-						const int32 ExternUniverseEnd = ExternUniverseStartValue + NumUniverses - 1;
-
-						if (Protocol->IsValidUniverseID(ExternUniverseStartValue) &&
-							Protocol->IsValidUniverseID(ExternUniverseEnd))
-						{
-							DMX_OVERRIDE_OUTPUTPORT_VAR(ExternUniverseStart, PortName, ExternUniverseStartValue);
-						}
-						else
+						const FString& ExternUniverseStartValueString = Args[1];
+						int32 ExternUniverseStartValue;
+						if (LexTryParseString<int32>(ExternUniverseStartValue, *ExternUniverseStartValueString))
 						{
-							UE_LOG(LogDMXProtocol, Warning, TEXT("Console command DMX.SetOutputPortExternUniverseStart failed. Extern Universe Start '%s' along with Num Universes '%i' results in a Universe range that is not supported by the Protocol of the Port."), *ExternUniverseStartValueString, NumUniverses);
+							const int32 NumUniverses = OutputPort->GetLocalUniverseEnd() - OutputPort->GetLocalUniverseStart() + 1;
+							const int32 ExternUniverseEnd = ExternUniverseStartValue + NumUniverses - 1;
+
+							if (Protocol->IsValidUniverseID(ExternUniverseStartValue) &&
+								Protocol->IsValidUniverseID(ExternUniverseEnd))
+							{
+								DMX_OVERRIDE_OUTPUTPORT_VAR(ExternUniverseStart, PortName, ExternUniverseStartValue);
+							}
+							else
+							{
+								UE_LOG(LogDMXProtocol, Warning, TEXT("Console command DMX.SetOutputPortExternUniverseStart failed. Extern Universe Start '%s' along with Num Universes '%i' results in a Universe range that is not supported by the Protocol of the Port."), *ExternUniverseStartValueString, NumUniverses);
+							}
 						}
 					}
 				}
-			}
-		})
-);
-
-static FAutoConsoleCommand GDMXSetOutputPortPriorityCommand(
-	TEXT("DMX.SetOutputPortPriority"),
-	TEXT("DMX.SetOutputPortPriority [PortName][Priority]. Sets the priority of the output port. Example: DMX.SetOutputPortPriority MyOutputPort 100"),
-	FConsoleCommandWithArgsDelegate::CreateStatic(
-		[](const TArray<FString>& Args)
-		{
-			constexpr int32 MinNumExpectedArgs = 2;
-			if (!UE::DMX::DMXOutputPort::VerifyConsoleCommandArguments(TEXT("DMX.SetOutputPortPriority"), MinNumExpectedArgs, Args))
+			})
+	);
+
+	static FAutoConsoleCommand GDMXSetOutputPortPriorityCommand(
+		TEXT("DMX.SetOutputPortPriority"),
+		TEXT("DMX.SetOutputPortPriority [PortName][Priority]. Sets the priority of the output port. Example: DMX.SetOutputPortPriority MyOutputPort 100"),
+		FConsoleCommandWithArgsDelegate::CreateStatic(
+			[](const TArray<FString>& Args)
 			{
-				return;
-			}
-
-			const FString& PortName = Args[0];
-			const FString& PriorityValueString = Args[1];
-			int32 PriorityValue;
-			if (LexTryParseString<int32>(PriorityValue, *PriorityValueString))
-			{
-				DMX_OVERRIDE_OUTPUTPORT_VAR(Priority, PortName, PriorityValue);
-			}
-		})
-);
-
-static FAutoConsoleCommand GDMXSetOutputPortDelayCommand(
-	TEXT("DMX.SetOutputPortDelay"),
-	TEXT("DMX.SetOutputPortDelay [PortName][Delay][(optional)FrameRate]. Sets the delay of the output port, optionally with a frame rate (e.g. '30fps', '0.001s' or '12000/1001'). Example: DMX.SetOutputPortDelay MyOutputPort 10.5 30fps"),
-	FConsoleCommandWithArgsDelegate::CreateStatic(
-		[](const TArray<FString>& Args)
-		{
-			constexpr int32 MinNumExpectedArgs = 2;
-			if (!UE::DMX::DMXOutputPort::VerifyConsoleCommandArguments(TEXT("DMX.SetOutputPortDelay"), MinNumExpectedArgs, Args))
-			{
-				return;
-			}
+				constexpr int32 MinNumExpectedArgs = 2;
+				if (!VerifyConsoleCommandArguments(TEXT("DMX.SetOutputPortPriority"), MinNumExpectedArgs, Args))
+				{
+					return;
+				}
 
-			const FString& PortName = Args[0];
-			const FString& DelayValueString = Args[1];
-			double DelayValue;
-			if (LexTryParseString<double>(DelayValue, *DelayValueString))
+				const FString& PortName = Args[0];
+				const FString& PriorityValueString = Args[1];
+				int32 PriorityValue;
+				if (LexTryParseString<int32>(PriorityValue, *PriorityValueString))
+				{
+					DMX_OVERRIDE_OUTPUTPORT_VAR(Priority, PortName, PriorityValue);
+				}
+			})
+	);
+
+	static FAutoConsoleCommand GDMXSetOutputPortDelayCommand(
+		TEXT("DMX.SetOutputPortDelay"),
+		TEXT("DMX.SetOutputPortDelay [PortName][Delay][(optional)FrameRate]. Sets the delay of the output port, optionally with a frame rate (e.g. '30fps', '0.001s' or '12000/1001'). Example: DMX.SetOutputPortDelay MyOutputPort 10.5 30fps"),
+		FConsoleCommandWithArgsDelegate::CreateStatic(
+			[](const TArray<FString>& Args)
 			{
-				FFrameRate NewFrameRate = FFrameRate(1, 1);
+				constexpr int32 MinNumExpectedArgs = 2;
+				if (!VerifyConsoleCommandArguments(TEXT("DMX.SetOutputPortDelay"), MinNumExpectedArgs, Args))
+				{
+					return;
+				}
 
-				// Parse optional frame rate argument
-				if (Args.Num() == 3)
+				const FString& PortName = Args[0];
+				const FString& DelayValueString = Args[1];
+				double DelayValue;
+				if (LexTryParseString<double>(DelayValue, *DelayValueString))
 				{
-					const FString& DelayFrameRateString = Args[2];
+					FFrameRate NewFrameRate = FFrameRate(1, 1);
 
-					TValueOrError<FFrameRate, FExpressionError> ParseResult = ParseFrameRate(*DelayFrameRateString);
-					if (ParseResult.IsValid())
+					// Parse optional frame rate argument
+					if (Args.Num() == 3)
 					{
-						NewFrameRate = ParseResult.GetValue();
-					}
-					else
-					{
-						UE_LOG(LogDMXProtocol, Warning, TEXT("Could not parse frame rate from console command DMX.SetOutputPortDelay. Note, not all decimal numbers are supported. E.g. 0.345s needs to be written as 345/1000."))
-					}
-				}
+						const FString& DelayFrameRateString = Args[2];
 
-				if (DelayValue >= 0.0)
-				{
-					const FDMXOutputPortSharedPtr OutputPort = UE::DMX::DMXOutputPort::FindOutputPortByName(PortName);
+						TValueOrError<FFrameRate, FExpressionError> ParseResult = ParseFrameRate(*DelayFrameRateString);
+						if (ParseResult.IsValid())
+						{
+							NewFrameRate = ParseResult.GetValue();
+						}
+						else
+						{
+							UE_LOG(LogDMXProtocol, Warning, TEXT("Could not parse frame rate from console command DMX.SetOutputPortDelay. Note, not all decimal numbers are supported. E.g. 0.345s needs to be written as 345/1000."))
+						}
+					}
 
-					if (OutputPort.IsValid())
+					if (DelayValue >= 0.0)
 					{
-						DMX_OVERRIDE_OUTPUTPORT_VAR(DelayFrameRate, PortName, NewFrameRate);
-						DMX_OVERRIDE_OUTPUTPORT_VAR(Delay, PortName, DelayValue);
+						const FDMXOutputPortSharedPtr OutputPort = FindOutputPortByName(PortName);
+
+						if (OutputPort.IsValid())
+						{
+							DMX_OVERRIDE_OUTPUTPORT_VAR(DelayFrameRate, PortName, NewFrameRate);
+							DMX_OVERRIDE_OUTPUTPORT_VAR(Delay, PortName, DelayValue);
+						}
 					}
 				}
-			}
-		})
-);
-
-static FAutoConsoleCommand GDMXResetOutputPortToProjectSettings(
-	TEXT("DMX.ResetOutputPortToProjectSettings"),
-	TEXT("DMX.ResetOutputPortToProjectSettings [PortName]. Resets the output port to how it is defined in project settings. Example: DMX.ResetOutputPortToProjectSettings MyOutputPort"),
-	FConsoleCommandWithArgsDelegate::CreateStatic(
-		[](const TArray<FString>& Args)
-		{
-			constexpr int32 MinNumExpectedArgs = 1;
-			if (!UE::DMX::DMXOutputPort::VerifyConsoleCommandArguments(TEXT("DMX.ResetOutputPortToProjectSettings"), MinNumExpectedArgs, Args))
+			})
+	);
+
+	static FAutoConsoleCommand GDMXResetOutputPortToProjectSettings(
+		TEXT("DMX.ResetOutputPortToProjectSettings"),
+		TEXT("DMX.ResetOutputPortToProjectSettings [PortName]. Resets the output port to how it is defined in project settings. Example: DMX.ResetOutputPortToProjectSettings MyOutputPort"),
+		FConsoleCommandWithArgsDelegate::CreateStatic(
+			[](const TArray<FString>& Args)
 			{
-				return;
-			}
+				constexpr int32 MinNumExpectedArgs = 1;
+				if (!VerifyConsoleCommandArguments(TEXT("DMX.ResetOutputPortToProjectSettings"), MinNumExpectedArgs, Args))
+				{
+					return;
+				}
 
-			const FString& PortName = Args[0];
-			const FDMXOutputPortSharedPtr OutputPort = UE::DMX::DMXOutputPort::FindOutputPortByName(PortName);
+				const FString& PortName = Args[0];
+				const FDMXOutputPortSharedPtr OutputPort = FindOutputPortByName(PortName);
 
-			if(OutputPort.IsValid())
-			{
-				const UDMXProtocolSettings* ProtocolSettings = GetDefault<UDMXProtocolSettings>();
-				if (ProtocolSettings)
+				if(OutputPort.IsValid())
 				{
-					const FDMXOutputPortConfig* PortConfigPtr = ProtocolSettings->OutputPortConfigs.FindByPredicate([OutputPort](const FDMXOutputPortConfig& OutputPortConfig)
-						{
-							return OutputPortConfig.GetPortGuid() == OutputPort->GetPortGuid();
-						});
-
-					if (PortConfigPtr)
+					const UDMXProtocolSettings* ProtocolSettings = GetDefault<UDMXProtocolSettings>();
+					if (ProtocolSettings)
 					{
-						FDMXOutputPortConfig PortConfig = *PortConfigPtr;
-						OutputPort->UpdateFromConfig(PortConfig);
+						const FDMXOutputPortConfig* PortConfigPtr = ProtocolSettings->OutputPortConfigs.FindByPredicate([OutputPort](const FDMXOutputPortConfig& OutputPortConfig)
+							{
+								return OutputPortConfig.GetPortGuid() == OutputPort->GetPortGuid();
+							});
+
+						if (PortConfigPtr)
+						{
+							FDMXOutputPortConfig PortConfig = *PortConfigPtr;
+							OutputPort->UpdateFromConfig(PortConfig);
+						}
 					}
 				}
-			}
-		})
-);
-
+			})
+	);
+}
 #undef DMX_OVERRIDE_OUTPUTPORT_VAR
 
-
 FDMXOutputPortSharedRef FDMXOutputPort::CreateFromConfig(FDMXOutputPortConfig& OutputPortConfig)
 {
 	// Port Configs are expected to have a valid guid always
@@ -510,7 +555,7 @@ FDMXOutputPortSharedRef FDMXOutputPort::CreateFromConfig(FDMXOutputPortConfig& O
 
 FDMXOutputPort::~FDMXOutputPort()
 {	
-	// All Inputs need to be explicitly removed before destruction 
+	// All Listeners need to be explicitly removed before destruction 
 	check(RawListeners.Num() == 0);
 	
 	// Port needs be unregistered before destruction
@@ -731,7 +776,7 @@ void FDMXOutputPort::AddRawListener(TSharedRef<FDMXRawListener> InRawListener)
 {
 	check(!RawListeners.Contains(InRawListener));
 
-	// Inputs need to run in the game thread
+	// Needs to run in the game thread
 	check(IsInGameThread());
 
 	RawListeners.Add(InRawListener);
