DMX - Auto Completed Device Addresses now are correctly set in the DMX protocols

diff --git a/Engine/Plugins/VirtualProduction/DMX/DMXProtocol/Source/DMXProtocol/Private/DMXProtocolUtils.cpp b/Engine/Plugins/VirtualProduction/DMX/DMXProtocol/Source/DMXProtocol/Private/DMXProtocolUtils.cpp
--- a/Engine/Plugins/VirtualProduction/DMX/DMXProtocol/Source/DMXProtocol/Private/DMXProtocolUtils.cpp
+++ b/Engine/Plugins/VirtualProduction/DMX/DMXProtocol/Source/DMXProtocol/Private/DMXProtocolUtils.cpp
@@ -53,6 +53,44 @@ TSharedPtr<FInternetAddr> FDMXProtocolUtils::CreateInternetAddr(const FString& I
 	return InternetAddr;
 }
 
+bool FDMXProtocolUtils::FindLocalNetworkInterfaceCardIPAddress(const FString& InIPAddressWithWildcards, FString& OutLocalNetworkInterfaceCardIPAddress)
+{
+	TArray<FString> SearchStrings;
+	constexpr bool bCullEmpty = true;
+	InIPAddressWithWildcards.ParseIntoArray(SearchStrings, TEXT("."), bCullEmpty);
+
+	const TArray<TSharedPtr<FString>> LocalNetworkInterfaceCardIPs = FDMXProtocolUtils::GetLocalNetworkInterfaceCardIPs();
+
+	FString NewIPAddress;
+	for (const TSharedPtr<FString>& NetworkInterfaceCardIP : LocalNetworkInterfaceCardIPs)
+	{
+		TArray<FString> NetworkInterfaceCardIPSubstrings;
+		NetworkInterfaceCardIP->ParseIntoArray(NetworkInterfaceCardIPSubstrings, TEXT("."));
+
+		// Try to match each search string
+		bool bMatchesSearchStrings = true;
+		for (int32 SubstringIndex = 0; SubstringIndex < NetworkInterfaceCardIPSubstrings.Num(); SubstringIndex++)
+		{
+			if (SearchStrings.IsValidIndex(SubstringIndex))
+			{
+				if (!NetworkInterfaceCardIPSubstrings[SubstringIndex].MatchesWildcard(SearchStrings[SubstringIndex]))
+				{
+					bMatchesSearchStrings = false;
+					break;
+				}
+			}
+		}
+
+		if (bMatchesSearchStrings)
+		{
+			OutLocalNetworkInterfaceCardIPAddress = *NetworkInterfaceCardIP;
+			return true;
+		}
+	}
+
+	return false;
+}
+
 FString FDMXProtocolUtils::GenerateUniqueNameFromExisting(const TSet<FString>& InExistingNames, const FString& InBaseName)
 {
 	if (!InBaseName.IsEmpty() && !InExistingNames.Contains(InBaseName))
diff --git a/Engine/Plugins/VirtualProduction/DMX/DMXProtocol/Source/DMXProtocol/Private/IO/DMXInputPortConfig.cpp b/Engine/Plugins/VirtualProduction/DMX/DMXProtocol/Source/DMXProtocol/Private/IO/DMXInputPortConfig.cpp
--- a/Engine/Plugins/VirtualProduction/DMX/DMXProtocol/Source/DMXProtocol/Private/IO/DMXInputPortConfig.cpp
+++ b/Engine/Plugins/VirtualProduction/DMX/DMXProtocol/Source/DMXProtocol/Private/IO/DMXInputPortConfig.cpp
@@ -137,16 +137,13 @@ FString FDMXInputPortConfig::GetDeviceAddress() const
 		return OverrideIP;
 	}
 
-	// Return the Auto Complete Device Address if that is enabled
+	// Return the Auto Complete Device Address if that's enabled
 	if (bAutoCompleteDeviceAddressEnabled)
 	{
-		const TArray<TSharedPtr<FString>> AvailableNetworkInterfaceCardIPs = FDMXProtocolUtils::GetLocalNetworkInterfaceCardIPs();
-		for (const TSharedPtr<FString>& NetworkInterfaceCardIP : AvailableNetworkInterfaceCardIPs)
+		FString NetworkInterfaceCardIPAddress;
+		if (FDMXProtocolUtils::FindLocalNetworkInterfaceCardIPAddress(AutoCompleteDeviceAddress, NetworkInterfaceCardIPAddress))
 		{
-			if ((*NetworkInterfaceCardIP).MatchesWildcard(AutoCompleteDeviceAddress))
-			{
-				return *NetworkInterfaceCardIP;
-			}
+			return NetworkInterfaceCardIPAddress;
 		}
 	}
 
diff --git a/Engine/Plugins/VirtualProduction/DMX/DMXProtocol/Source/DMXProtocol/Private/IO/DMXOutputPortConfig.cpp b/Engine/Plugins/VirtualProduction/DMX/DMXProtocol/Source/DMXProtocol/Private/IO/DMXOutputPortConfig.cpp
--- a/Engine/Plugins/VirtualProduction/DMX/DMXProtocol/Source/DMXProtocol/Private/IO/DMXOutputPortConfig.cpp
+++ b/Engine/Plugins/VirtualProduction/DMX/DMXProtocol/Source/DMXProtocol/Private/IO/DMXOutputPortConfig.cpp
@@ -155,16 +155,13 @@ FString FDMXOutputPortConfig::GetDeviceAddress() const
 		return OverrideIP;
 	}
 
-	// Return the Auto Complete Device Address if that is enabled
+	// Return the Auto Complete Device Address if that's enabled
 	if (bAutoCompleteDeviceAddressEnabled)
 	{
-		const TArray<TSharedPtr<FString>> AvailableNetworkInterfaceCardIPs = FDMXProtocolUtils::GetLocalNetworkInterfaceCardIPs();
-		for (const TSharedPtr<FString>& NetworkInterfaceCardIP : AvailableNetworkInterfaceCardIPs)
+		FString NetworkInterfaceCardIPAddress;
+		if (FDMXProtocolUtils::FindLocalNetworkInterfaceCardIPAddress(AutoCompleteDeviceAddress, NetworkInterfaceCardIPAddress))
 		{
-			if ((*NetworkInterfaceCardIP).MatchesWildcard(AutoCompleteDeviceAddress))
-			{
-				return *NetworkInterfaceCardIP;
-			}
+			return NetworkInterfaceCardIPAddress;
 		}
 	}
 
diff --git a/Engine/Plugins/VirtualProduction/DMX/DMXProtocol/Source/DMXProtocol/Private/IO/DMXPort.cpp b/Engine/Plugins/VirtualProduction/DMX/DMXProtocol/Source/DMXProtocol/Private/IO/DMXPort.cpp
--- a/Engine/Plugins/VirtualProduction/DMX/DMXProtocol/Source/DMXProtocol/Private/IO/DMXPort.cpp
+++ b/Engine/Plugins/VirtualProduction/DMX/DMXProtocol/Source/DMXProtocol/Private/IO/DMXPort.cpp
@@ -54,17 +54,16 @@ bool FDMXPort::IsValidPortSlow() const
 	ISocketSubsystem* SocketSubsystem = ISocketSubsystem::Get(PLATFORM_SOCKETSUBSYSTEM);
 	TSharedPtr<FInternetAddr> IPInternetAddr = SocketSubsystem->CreateInternetAddr();
 
-	int32 FinalUniverseStart = LocalUniverseStart + ExternUniverseStart;
-	if (FinalUniverseStart < Protocol->GetMinUniverseID())
+	if (ExternUniverseStart < Protocol->GetMinUniverseID())
 	{
-		UE_LOG(LogDMXProtocol, Warning, TEXT("Cannot create DMX Port: Resulting First Universe %i is smaller than Protocol Min %i."), FinalUniverseStart, Protocol->GetMinUniverseID());
+		UE_LOG(LogDMXProtocol, Warning, TEXT("Cannot create DMX Port: Resulting First Universe %i is smaller than Protocol Min %i."), ExternUniverseStart, Protocol->GetMinUniverseID());
 		return false;
 	}
 
-	int32 FinalUniverseEnd = FinalUniverseStart + NumUniverses - 1;
-	if (FinalUniverseEnd > Protocol->GetMaxUniverseID())
+	const int32 ExternUniverseEnd = ExternUniverseStart + NumUniverses - 1;
+	if (ExternUniverseEnd > Protocol->GetMaxUniverseID())
 	{
-		UE_LOG(LogDMXProtocol, Warning, TEXT("Cannot create DMX Port: Resulting Last Universe %i is bigger than Protocol Max %i."), FinalUniverseStart, Protocol->GetMinUniverseID());
+		UE_LOG(LogDMXProtocol, Warning, TEXT("Cannot create DMX Port: Resulting Last Universe %i is bigger than Protocol Max %i."), ExternUniverseEnd, Protocol->GetMinUniverseID());
 		return false;
 	}
 
diff --git a/Engine/Plugins/VirtualProduction/DMX/DMXProtocol/Source/DMXProtocol/Public/DMXProtocolUtils.h b/Engine/Plugins/VirtualProduction/DMX/DMXProtocol/Source/DMXProtocol/Public/DMXProtocolUtils.h
--- a/Engine/Plugins/VirtualProduction/DMX/DMXProtocol/Source/DMXProtocol/Public/DMXProtocolUtils.h
+++ b/Engine/Plugins/VirtualProduction/DMX/DMXProtocol/Source/DMXProtocol/Public/DMXProtocolUtils.h
@@ -23,6 +23,9 @@ public:
 	 */
 	static TSharedPtr<FInternetAddr> CreateInternetAddr(const FString& IPAddress, int32 Port);
 
+	/** Tries to find a local network interface card IP that matches an IP Address with wildcards */
+	static bool FindLocalNetworkInterfaceCardIPAddress(const FString& InIPAddressWithWildcards, FString& OutLocalNetworkInterfaceCardIPAddress);
+
 	/**
 	 * Generates a unique name given a base one and a list of existing ones, by appending an index to
 	 * existing names. If InBaseName is an empty String, it returns "Default name".
diff --git a/Engine/Plugins/VirtualProduction/DMX/DMXProtocol/Source/DMXProtocolEditor/Private/DetailsCustomizations/DMXInputPortConfigCustomization.cpp b/Engine/Plugins/VirtualProduction/DMX/DMXProtocol/Source/DMXProtocolEditor/Private/DetailsCustomizations/DMXInputPortConfigCustomization.cpp
--- a/Engine/Plugins/VirtualProduction/DMX/DMXProtocol/Source/DMXProtocolEditor/Private/DetailsCustomizations/DMXInputPortConfigCustomization.cpp
+++ b/Engine/Plugins/VirtualProduction/DMX/DMXProtocol/Source/DMXProtocolEditor/Private/DetailsCustomizations/DMXInputPortConfigCustomization.cpp
@@ -323,46 +323,20 @@ void FDMXInputPortConfigCustomization::UpdateAutoCompleteDeviceAddressTextBox()
 		FString AutoCompleteDeviceAddress;
 		if (AutoCompleteDeviceAddressHandle->GetValue(AutoCompleteDeviceAddress) == FPropertyAccess::Success)
 		{
-			TArray<FString> SearchStrings;
-			constexpr bool bCullEmpty = true;
-			AutoCompleteDeviceAddress.ParseIntoArray(SearchStrings, TEXT("."), bCullEmpty);
-
-			const TArray<TSharedPtr<FString>> LocalNetworkInterfaceCardIPs = FDMXProtocolUtils::GetLocalNetworkInterfaceCardIPs();
-
-			FString NewIPAddress;
-			for (const TSharedPtr<FString>& NetworkInterfaceCardIP : LocalNetworkInterfaceCardIPs)
+			FString NetworkInterfaceCardIPAddress;
+			if (FDMXProtocolUtils::FindLocalNetworkInterfaceCardIPAddress(AutoCompleteDeviceAddress, NetworkInterfaceCardIPAddress))
 			{
-				TArray<FString> NetworkInterfaceCardIPSubstrings;
-				(*NetworkInterfaceCardIP).ParseIntoArray(NetworkInterfaceCardIPSubstrings, TEXT("."));
-				
-				// Try to match each search string
-				bool bMatchesSearchStrings = true;
-				for (int32 SubstringIndex = 0; SubstringIndex < NetworkInterfaceCardIPSubstrings.Num(); SubstringIndex++)
-				{
-					if (SearchStrings.IsValidIndex(SubstringIndex))
-					{
-						if (!NetworkInterfaceCardIPSubstrings[SubstringIndex].MatchesWildcard(SearchStrings[SubstringIndex]))
-						{
-							bMatchesSearchStrings = false;
-							break;
-						}
-					}
-				}
-
-				if (bMatchesSearchStrings)
-				{
-					// Set the text and return if a match was found
-					const FText IPAddressText = FText::FromString(*NetworkInterfaceCardIP);
-					AutoCompletedDeviceAddressTextBlock->SetText(FText::Format(LOCTEXT("AutoCompletedIPAddressText", "{0} (auto-completed)"), IPAddressText));
-					return;
-				}
+				const FText IPAddressText = FText::FromString(*NetworkInterfaceCardIPAddress);
+				AutoCompletedDeviceAddressTextBlock->SetText(FText::Format(LOCTEXT("AutoCompletedIPAddressText", "{0} (auto-completed)"), IPAddressText));
+			}
+			else
+			{
+				// Fall back to the manually set one
+				FString DeviceAddress;
+				DeviceAddressHandle->GetValue(DeviceAddress);
+				const FText NoMatchForIPAddressText = FText::Format(LOCTEXT("NoRegexMatchForDeviceAddress", "No match. Using '{0}'"), FText::FromString(DeviceAddress));
+				AutoCompletedDeviceAddressTextBlock->SetText(NoMatchForIPAddressText);
 			}
-
-			// Fall back to the manually set one
-			FString DeviceAddress;
-			DeviceAddressHandle->GetValue(DeviceAddress);
-			const FText NoMatchForIPAddressText = FText::Format(LOCTEXT("NoRegexMatchForDeviceAddress", "No match. Using '{0}'"), FText::FromString(DeviceAddress));
-			AutoCompletedDeviceAddressTextBlock->SetText(NoMatchForIPAddressText);
 		}
 	}
 }
diff --git a/Engine/Plugins/VirtualProduction/DMX/DMXProtocol/Source/DMXProtocolEditor/Private/DetailsCustomizations/DMXOutputPortConfigCustomization.cpp b/Engine/Plugins/VirtualProduction/DMX/DMXProtocol/Source/DMXProtocolEditor/Private/DetailsCustomizations/DMXOutputPortConfigCustomization.cpp
--- a/Engine/Plugins/VirtualProduction/DMX/DMXProtocol/Source/DMXProtocolEditor/Private/DetailsCustomizations/DMXOutputPortConfigCustomization.cpp
+++ b/Engine/Plugins/VirtualProduction/DMX/DMXProtocol/Source/DMXProtocolEditor/Private/DetailsCustomizations/DMXOutputPortConfigCustomization.cpp
@@ -361,46 +361,20 @@ void FDMXOutputPortConfigCustomization::UpdateAutoCompleteDeviceAddressTextBox()
 		FString AutoCompleteDeviceAddress;
 		if (AutoCompleteDeviceAddressHandle->GetValue(AutoCompleteDeviceAddress) == FPropertyAccess::Success)
 		{
-			TArray<FString> SearchStrings;
-			constexpr bool bCullEmpty = true;
-			AutoCompleteDeviceAddress.ParseIntoArray(SearchStrings, TEXT("."), bCullEmpty);
-
-			const TArray<TSharedPtr<FString>> LocalNetworkInterfaceCardIPs = FDMXProtocolUtils::GetLocalNetworkInterfaceCardIPs();
-
-			FString NewIPAddress;
-			for (const TSharedPtr<FString>& NetworkInterfaceCardIP : LocalNetworkInterfaceCardIPs)
+			FString NetworkInterfaceCardIPAddress;
+			if (FDMXProtocolUtils::FindLocalNetworkInterfaceCardIPAddress(AutoCompleteDeviceAddress, NetworkInterfaceCardIPAddress))
 			{
-				TArray<FString> NetworkInterfaceCardIPSubstrings;
-				(*NetworkInterfaceCardIP).ParseIntoArray(NetworkInterfaceCardIPSubstrings, TEXT("."));
-
-				// Try to match each search string
-				bool bMatchesSearchStrings = true;
-				for (int32 SubstringIndex = 0; SubstringIndex < NetworkInterfaceCardIPSubstrings.Num(); SubstringIndex++)
-				{
-					if (SearchStrings.IsValidIndex(SubstringIndex))
-					{
-						if (!NetworkInterfaceCardIPSubstrings[SubstringIndex].MatchesWildcard(SearchStrings[SubstringIndex]))
-						{
-							bMatchesSearchStrings = false;
-							break;
-						}
-					}
-				}
-
-				if (bMatchesSearchStrings)
-				{
-					// Set the text and return if a match was found
-					const FText IPAddressText = FText::FromString(*NetworkInterfaceCardIP);
-					AutoCompletedDeviceAddressTextBlock->SetText(FText::Format(LOCTEXT("AutoCompletedIPAddressText", "{0} (auto-completed)"), IPAddressText));
-					return;
-				}
+				const FText IPAddressText = FText::FromString(*NetworkInterfaceCardIPAddress);
+				AutoCompletedDeviceAddressTextBlock->SetText(FText::Format(LOCTEXT("AutoCompletedIPAddressText", "{0} (auto-completed)"), IPAddressText));
+			}
+			else
+			{
+				// Fall back to the manually set one
+				FString DeviceAddress;
+				DeviceAddressHandle->GetValue(DeviceAddress);
+				const FText NoMatchForIPAddressText = FText::Format(LOCTEXT("NoRegexMatchForDeviceAddress", "No match. Using '{0}'"), FText::FromString(DeviceAddress));
+				AutoCompletedDeviceAddressTextBlock->SetText(NoMatchForIPAddressText);
 			}
-
-			// Fall back to the manually set one
-			FString DeviceAddress;
-			DeviceAddressHandle->GetValue(DeviceAddress);
-			const FText NoMatchForIPAddressText = FText::Format(LOCTEXT("NoRegexMatchForDeviceAddress", "No match. Using '{0}'"), FText::FromString(DeviceAddress));
-			AutoCompletedDeviceAddressTextBlock->SetText(NoMatchForIPAddressText);
 		}
 	}
 }
