DMX - Fix sACN multicast loopback

diff --git a/Engine/Plugins/VirtualProduction/DMX/DMXProtocol/Source/DMXProtocolSACN/Private/DMXProtocolSACNReceiver.cpp b/Engine/Plugins/VirtualProduction/DMX/DMXProtocol/Source/DMXProtocolSACN/Private/DMXProtocolSACNReceiver.cpp
--- a/Engine/Plugins/VirtualProduction/DMX/DMXProtocol/Source/DMXProtocolSACN/Private/DMXProtocolSACNReceiver.cpp
+++ b/Engine/Plugins/VirtualProduction/DMX/DMXProtocol/Source/DMXProtocolSACN/Private/DMXProtocolSACNReceiver.cpp
@@ -58,19 +58,19 @@ FDMXProtocolSACNReceiver::~FDMXProtocolSACNReceiver()
 
 TSharedPtr<FDMXProtocolSACNReceiver> FDMXProtocolSACNReceiver::TryCreate(const TSharedPtr<FDMXProtocolSACN, ESPMode::ThreadSafe>& SACNProtocol, const FString& IPAddress)
 {
-	TSharedPtr<FInternetAddr> EndpointInternetAddr = FDMXProtocolUtils::CreateInternetAddr(IPAddress, ACN_PORT);
+	const TSharedPtr<FInternetAddr> EndpointInternetAddr = FDMXProtocolUtils::CreateInternetAddr(IPAddress, ACN_PORT);
 	if (!EndpointInternetAddr.IsValid())
 	{
 		UE_LOG(LogDMXProtocol, Error, TEXT("Cannot create sACN receiver: Invalid IP address: %s"), *IPAddress);
 		return nullptr;
 	}
-
-	FIPv4Endpoint Endpoint = FIPv4Endpoint(EndpointInternetAddr);
+	const FIPv4Endpoint Endpoint(EndpointInternetAddr);
+	const FIPv4Endpoint AnyEndpoint(FIPv4Address::Any, ACN_PORT);
 
 	FSocket* NewListeningSocket = FUdpSocketBuilder(TEXT("UDPSACNListeningSocket"))
 		.AsBlocking()
 		.AsReusable()
-		.BoundToEndpoint(Endpoint)
+		.BoundToEndpoint(AnyEndpoint)
 		.WithMulticastLoopback()
 		.WithMulticastTtl(1)
 		.WithMulticastInterface(Endpoint.Address);
diff --git a/Engine/Plugins/VirtualProduction/DMX/DMXProtocol/Source/DMXProtocolSACN/Private/DMXProtocolSACNSender.cpp b/Engine/Plugins/VirtualProduction/DMX/DMXProtocol/Source/DMXProtocolSACN/Private/DMXProtocolSACNSender.cpp
--- a/Engine/Plugins/VirtualProduction/DMX/DMXProtocol/Source/DMXProtocolSACN/Private/DMXProtocolSACNSender.cpp
+++ b/Engine/Plugins/VirtualProduction/DMX/DMXProtocol/Source/DMXProtocolSACN/Private/DMXProtocolSACNSender.cpp
@@ -102,25 +102,20 @@ TSharedPtr<FDMXProtocolSACNSender> FDMXProtocolSACNSender::TryCreateMulticastSen
 	const FString& InNetworkInterfaceIP)
 {
 	// Try to create a socket
-	TSharedPtr<FInternetAddr> NewNetworkInterfaceInternetAddr = FDMXProtocolUtils::CreateInternetAddr(InNetworkInterfaceIP, ACN_PORT);
+	const TSharedPtr<FInternetAddr> NewNetworkInterfaceInternetAddr = FDMXProtocolUtils::CreateInternetAddr(InNetworkInterfaceIP, ACN_PORT);
 	if (!NewNetworkInterfaceInternetAddr.IsValid())
 	{
 		UE_LOG(LogDMXProtocol, Error, TEXT("Cannot create sACN sender: Invalid IP address: %s"), *InNetworkInterfaceIP);
 		return nullptr;
 	}
 
-	FIPv4Endpoint NewNetworkInterfaceEndpoint = FIPv4Endpoint(NewNetworkInterfaceInternetAddr);
-
 	FSocket* NewSocket = 
 		FUdpSocketBuilder(TEXT("UDPSACNMulticastSocket"))
-#if PLATFORM_SUPPORTS_UDP_MULTICAST_GROUP
-		.WithMulticastInterface(NewNetworkInterfaceEndpoint.Address)
-		.WithMulticastTtl(1)
-#endif
-		.AsBlocking()
-		.BoundToEndpoint(NewNetworkInterfaceInternetAddr)
-		.AsReusable();
-
+        .AsBlocking()
+        .AsReusable()
+        .WithMulticastLoopback()
+        .WithMulticastTtl(1);
+	
 	if(!NewSocket)
 	{
 		UE_LOG(LogDMXProtocol, Error, TEXT("Invalid Network Interface IP %s for DMX Port. Please update your Output Ports in Project Settings -> Plugins -> DMX Plugin"), *InNetworkInterfaceIP);
