DMX - Default the 'All Fixture Patches receive DMX in Editor' setting to true

diff --git a/Engine/Plugins/VirtualProduction/DMX/DMXProtocol/Source/DMXProtocol/Private/DMXProtocolSettings.cpp b/Engine/Plugins/VirtualProduction/DMX/DMXProtocol/Source/DMXProtocol/Private/DMXProtocolSettings.cpp
--- a/Engine/Plugins/VirtualProduction/DMX/DMXProtocol/Source/DMXProtocol/Private/DMXProtocolSettings.cpp
+++ b/Engine/Plugins/VirtualProduction/DMX/DMXProtocol/Source/DMXProtocol/Private/DMXProtocolSettings.cpp
@@ -16,10 +16,6 @@
 UDMXProtocolSettings::UDMXProtocolSettings()
 	: SendingRefreshRate(DMX_RATE)
 	, ReceivingRefreshRate_DEPRECATED(DMX_RATE)
-	, bDefaultSendDMXEnabled(true)
-	, bDefaultReceiveDMXEnabled(true)
-	, bOverrideSendDMXEnabled(true)	
-	, bOverrideReceiveDMXEnabled(true)
 {
 	FixtureCategories =
 	{
diff --git a/Engine/Plugins/VirtualProduction/DMX/DMXProtocol/Source/DMXProtocol/Public/DMXProtocolSettings.h b/Engine/Plugins/VirtualProduction/DMX/DMXProtocol/Source/DMXProtocol/Public/DMXProtocolSettings.h
--- a/Engine/Plugins/VirtualProduction/DMX/DMXProtocol/Source/DMXProtocol/Public/DMXProtocolSettings.h
+++ b/Engine/Plugins/VirtualProduction/DMX/DMXProtocol/Source/DMXProtocol/Public/DMXProtocolSettings.h
@@ -156,17 +156,17 @@ public:
 private:
 	/** Whether DMX is sent to the network. Recalled whenever editor or game starts.  */
 	UPROPERTY(Config, EditAnywhere, Category = "DMX|Communication Settings", Meta = (AllowPrivateAccess = true, DisplayName = "Send DMX by default"))
-	bool bDefaultSendDMXEnabled;
+	bool bDefaultSendDMXEnabled = true;
 
 	/** Whether DMX is received from the network. Recalled whenever editor or game starts. */
 	UPROPERTY(Config, EditAnywhere, Category = "DMX|Communication Settings", Meta = (AllowPrivateAccess = true, DisplayName = "Receive DMX by default"))
-	bool bDefaultReceiveDMXEnabled;
+	bool bDefaultReceiveDMXEnabled = true;
 
 	/** Overrides the default bDefaultSendDMXEnabled value at runtime */
-	bool bOverrideSendDMXEnabled;
+	bool bOverrideSendDMXEnabled = true;
 
 	/** Overrides the default bDefaultReceiveDMXEnabled value at runtime */
-	bool bOverrideReceiveDMXEnabled;
+	bool bOverrideReceiveDMXEnabled = true;
 
 	/** Broadcast when send DMX is enabled or disabled */
 	FDMXOnSendDMXEnabled OnSetSendDMXEnabledDelegate;
@@ -183,7 +183,7 @@ private:
 #if WITH_EDITORONLY_DATA
 	/** If true, all fixture patches receive DMX in Editor. This overrides the fixture patches 'Receive DMX In Editor' property. */
 	UPROPERTY(Config, EditAnywhere, Category = "DMX|Communication Settings", Meta = (DisplayName = "All Fixture Patches receive DMX in Editor"))
-	bool bAllFixturePatchesReceiveDMXInEditor = false;
+	bool bAllFixturePatchesReceiveDMXInEditor = true;
 
 	/** Broadcast when Fixture Patches recieve DMX in Editor is enabled or disabled */
 	FDMXOnAllFixturePatchesReceiveDMXInEditorEnabled OnAllFixturePatchesReceiveDMXInEditorEnabled;
