DMX - Revise synchronization and thread-safety of DMX IO

diff --git a/Engine/Plugins/VirtualProduction/DMX/DMXProtocol/Source/DMXProtocol/Private/IO/DMXInputPort.cpp b/Engine/Plugins/VirtualProduction/DMX/DMXProtocol/Source/DMXProtocol/Private/IO/DMXInputPort.cpp
--- a/Engine/Plugins/VirtualProduction/DMX/DMXProtocol/Source/DMXProtocol/Private/IO/DMXInputPort.cpp
+++ b/Engine/Plugins/VirtualProduction/DMX/DMXProtocol/Source/DMXProtocol/Private/IO/DMXInputPort.cpp
@@ -598,16 +598,17 @@ bool FDMXInputPort::CheckPriority(const int32 InPriority)
 
 void FDMXInputPort::AddRawListener(TSharedRef<FDMXRawListener> InRawListener)
 {
-	check(!RawListeners.Contains(InRawListener));
+	const FScopeLock LockRawListeners(&AccessRawListenersMutex);
 
-	// Inputs need to run in the game thread
-	check(IsInGameThread());
+	check(!RawListeners.Contains(InRawListener));
 
 	RawListeners.Add(InRawListener);
 }
 
 void FDMXInputPort::RemoveRawListener(TSharedRef<FDMXRawListener> InRawListenerToRemove)
 {
+	const FScopeLock LockRawListeners(&AccessRawListenersMutex);
+
 	RawListeners.Remove(InRawListenerToRemove);
 }
 
@@ -670,10 +671,12 @@ void FDMXInputPort::ClearBuffers()
 	check(IsInGameThread());
 #endif // UE_BUILD_DEBUG
 
+	AccessRawListenersMutex.Lock();
 	for (const TSharedRef<FDMXRawListener>& RawListener : RawListeners)
 	{
 		RawListener->ClearBuffer();
 	}
+	AccessRawListenersMutex.Unlock();
 
 	DefaultInputQueue.Empty();
 	ExternUniverseToLatestSignalMap.Reset();
@@ -686,10 +689,12 @@ void FDMXInputPort::InputDMXSignal(const FDMXSignalSharedRef& DMXSignal)
 		int32 ExternUniverseID = DMXSignal->ExternUniverseID;
 		if (IsExternUniverseInPortRange(ExternUniverseID))
 		{
+			AccessRawListenersMutex.Lock();
 			for (const TSharedRef<FDMXRawListener>& RawListener : RawListeners)
 			{
 				RawListener->EnqueueSignal(this, DMXSignal);
 			}
+			AccessRawListenersMutex.Unlock();
 
 			if (bUseDefaultInputQueue)
 			{
@@ -708,10 +713,12 @@ void FDMXInputPort::SingleProducerInputDMXSignal(const FDMXSignalSharedRef& DMXS
 		int32 ExternUniverseID = DMXSignal->ExternUniverseID;
 		if (IsExternUniverseInPortRange(ExternUniverseID))
 		{
+			AccessRawListenersMutex.Lock();
 			for (const TSharedRef<FDMXRawListener>& RawListener : RawListeners)
 			{
 				RawListener->EnqueueSignal(this, DMXSignal);
 			}
+			AccessRawListenersMutex.Unlock();
 
 			if (bUseDefaultInputQueue)
 			{
diff --git a/Engine/Plugins/VirtualProduction/DMX/DMXProtocol/Source/DMXProtocol/Private/IO/DMXOutputPort.cpp b/Engine/Plugins/VirtualProduction/DMX/DMXProtocol/Source/DMXProtocol/Private/IO/DMXOutputPort.cpp
--- a/Engine/Plugins/VirtualProduction/DMX/DMXProtocol/Source/DMXProtocol/Private/IO/DMXOutputPort.cpp
+++ b/Engine/Plugins/VirtualProduction/DMX/DMXProtocol/Source/DMXProtocol/Private/IO/DMXOutputPort.cpp
@@ -11,6 +11,7 @@
 #include "IO/DMXPortManager.h"
 #include "IO/DMXRawListener.h"
 
+#include "HAL/Event.h"
 #include "HAL/IConsoleManager.h"
 #include "HAL/RunnableThread.h"
 #include "Misc/FrameRate.h"
@@ -527,6 +528,8 @@ namespace UE::DMX::DMXOutputPort::Private::ConsoleCommands
 
 FDMXOutputPortSharedRef FDMXOutputPort::CreateFromConfig(FDMXOutputPortConfig& OutputPortConfig)
 {
+	checkf(IsInGameThread(), TEXT("Only game-thread can create ports."));
+
 	// Port Configs are expected to have a valid guid always
 	check(OutputPortConfig.GetPortGuid().IsValid());
 
@@ -572,6 +575,12 @@ FDMXOutputPort::~FDMXOutputPort()
 
 FDMXOutputPortConfig FDMXOutputPort::MakeOutputPortConfig() const
 {
+	checkf(IsInGameThread(), TEXT("Only game-thread can read port configs."));
+
+	// Acquire locks in order as they're declared in the header
+	const FScopeLock LockDestinationAddresses(&AccessDestinationAddressesMutex);
+	const FScopeLock LockCommunicationDeterminator(&AccessCommunicationDeterminatorMutex);
+
 	FDMXOutputPortConfigParams Params;
 	Params.PortName = PortName;
 	Params.ProtocolName = Protocol.IsValid() ? Protocol->GetProtocolName() : NAME_None;
@@ -593,6 +602,12 @@ FDMXOutputPortConfig FDMXOutputPort::MakeOutputPortConfig() const
 
 void FDMXOutputPort::UpdateFromConfig(FDMXOutputPortConfig& InOutOutputPortConfig, bool bForceUpdateRegistrationWithProtocol)
 {	
+	checkf(IsInGameThread(), TEXT("Only game-thread can update a port from a config."));
+
+	// Acquire locks in order as they're declared in the header
+	const FScopeLock LockDestinationAddresses(&AccessDestinationAddressesMutex);
+	const FScopeLock LockCommunicationDeterminator(&AccessCommunicationDeterminatorMutex);
+
 	// Need a valid config for the port
 	InOutOutputPortConfig.MakeValid();
 
@@ -684,29 +699,30 @@ const FGuid& FDMXOutputPort::GetPortGuid() const
 
 void FDMXOutputPort::ClearBuffers()
 {
-	check(IsInGameThread());
+	// Note, it is important that this is called from the game thread!
+	// Otherwise the Send methods would need a lock on enqueuing, deterring the lock-free implementation.
+	checkf(IsInGameThread(), TEXT("Only the game-thread can ever clear buffers."));
 
-	for (const TSharedRef<FDMXRawListener>& RawListener : RawListeners)
-	{
-		RawListener->ClearBuffer();
-	}
+	// Acquire locks in order as they're declared in the header
+	const FScopeLock LockSignalFragments(&AccessSignalFragmentsMutex);
+	const FScopeLock LockExternUniverseToLatestSignalMap_PortThread(&AccessExternUniverseToLatestSignalMap_PortThreadMutex);
+	const FScopeLock LockRawlisteners(&AccessRawListenersMutex);
 
-	// Clear gamethread buffer
 	ExternUniverseToLatestSignalMap_GameThread.Reset();
+	ExternUniverseToLatestSignalMap_PortThread.Reset();
 
-	// Clear port thread buffers
-	FScopeLock LockClearBuffers(&ClearBuffersCriticalSection);
 	SignalFragments.Empty();
-	ExternUniverseToLatestSignalMap_PortThread.Reset();
-}
 
-bool FDMXOutputPort::IsLoopbackToEngine() const
-{
-	return CommunicationDeterminator.NeedsLoopbackToEngine();
+	for (const TSharedRef<FDMXRawListener>& RawListener : RawListeners)
+	{
+		RawListener->ClearBuffer();
+	}
 }
 
 TArray<FString> FDMXOutputPort::GetDestinationAddresses() const
 {
+	const FScopeLock LockDestinationAddresses(&AccessDestinationAddressesMutex);
+
 	TArray<FString> Result;
 	Result.Reserve(DestinationAddresses.Num());
 	for (const FDMXOutputPortDestinationAddress& DestinationAddress : DestinationAddresses)
@@ -721,6 +737,7 @@ bool FDMXOutputPort::GameThreadGetDMXSignal(int32 LocalUniverseID, FDMXSignalSha
 {
 	check(IsInGameThread());
 
+	// We don't mind to lock the communication determinator, it is ok if it is out of sync while changing.
 	const bool bNeedsLoopbackToEngine = CommunicationDeterminator.NeedsLoopbackToEngine();
 	if (bNeedsLoopbackToEngine || bEvenIfNotLoopbackToEngine)
 	{
@@ -742,6 +759,7 @@ bool FDMXOutputPort::GameThreadGetDMXSignalFromRemoteUniverse(FDMXSignalSharedPt
 	// DEPRECATED 4.27
 	check(IsInGameThread());
 
+	// We don't mind to lock the communication determinator, it is ok if it is out of sync while changing
 	const bool bNeedsLoopbackToEngine = CommunicationDeterminator.NeedsLoopbackToEngine();
 	if (bNeedsLoopbackToEngine || bEvenIfNotLoopbackToEngine)
 	{
@@ -756,14 +774,9 @@ bool FDMXOutputPort::GameThreadGetDMXSignalFromRemoteUniverse(FDMXSignalSharedPt
 	return false;
 }
 
-FString FDMXOutputPort::GetDestinationAddress() const
-{
-	// DEPRECATED 5.0
-	return DestinationAddresses.Num() > 0 ? DestinationAddresses[0].DestinationAddressString : TEXT("");
-}
-
 bool FDMXOutputPort::IsRegistered() const
-{
+{	
+	const FScopeLock LockSenderArray(&AccessDMXSenderArrayMutex);
 	if (DMXSenderArray.Num() > 0)
 	{
 		return true;
@@ -774,6 +787,8 @@ bool FDMXOutputPort::IsRegistered() const
 
 void FDMXOutputPort::AddRawListener(TSharedRef<FDMXRawListener> InRawListener)
 {
+	const FScopeLock LockRawListeners(&AccessRawListenersMutex);
+
 	check(!RawListeners.Contains(InRawListener));
 
 	// Needs to run in the game thread
@@ -784,15 +799,18 @@ void FDMXOutputPort::AddRawListener(TSharedRef<FDMXRawListener> InRawListener)
 
 void FDMXOutputPort::RemoveRawListener(TSharedRef<FDMXRawListener> InRawListenerToRemove)
 {
+	const FScopeLock LockRawListeners(&AccessRawListenersMutex);
+
 	RawListeners.Remove(InRawListenerToRemove);
 }
 
 void FDMXOutputPort::SendDMX(int32 LocalUniverseID, const TMap<int32, uint8>& ChannelToValueMap)
 {
-	check(IsInGameThread());
+	checkf(IsInGameThread(), TEXT("Only the game-thread can Send from a DMX Output Port."));
 
 	if (IsLocalUniverseInPortRange(LocalUniverseID))
 	{
+		// We don't mind to lock the communication determinator, it is ok if it is not sync while changing
 		const bool bNeedsSendDMX = CommunicationDeterminator.NeedsSendDMX();
 		const bool bNeedsLoopbackToEngine = CommunicationDeterminator.NeedsLoopbackToEngine();
 
@@ -803,8 +821,8 @@ void FDMXOutputPort::SendDMX(int32 LocalUniverseID, const TMap<int32, uint8>& Ch
 
 			const double SendTime = FPlatformTime::Seconds() + DelaySeconds;
 
-			// Enqueue for this port's thread
-			const TSharedPtr<FDMXSignalFragment> Fragment = MakeShared<FDMXSignalFragment>(ExternUniverseID, ChannelToValueMap, SendTime);
+			// Enqueue for this port's thread. Thread-safe without locking.
+			const TSharedPtr<FDMXSignalFragment> Fragment = MakeShared<FDMXSignalFragment>(ExternUniverseID, ChannelToValueMap, SendTime);	
 			SignalFragments.Enqueue(Fragment);
 
 			// Write the fragment to the game thread's buffer
@@ -832,6 +850,7 @@ void FDMXOutputPort::SendDMXToRemoteUniverse(const TMap<int32, uint8>& ChannelTo
 
 	if (IsExternUniverseInPortRange(RemoteUniverse))
 	{
+		// We don't mind to lock the communication determinator, it is ok if it is true while changing
 		const bool bNeedsSendDMX = CommunicationDeterminator.NeedsSendDMX();
 		const bool bNeedsLoopbackToEngine = CommunicationDeterminator.NeedsLoopbackToEngine();
 
@@ -867,24 +886,33 @@ void FDMXOutputPort::SendDMXToRemoteUniverse(const TMap<int32, uint8>& ChannelTo
 
 bool FDMXOutputPort::Register()
 {
+	// Acquire locks in order as they're declared in the header
+	const FScopeLock LockAccessSenderArray(&AccessDMXSenderArrayMutex);
+	const FScopeLock LockCommunicationDeterminator(&AccessCommunicationDeterminatorMutex);
+
 	if (Protocol.IsValid() && IsValidPortSlow() && CommunicationDeterminator.IsSendDMXEnabled() && !FDMXPortManager::Get().AreProtocolsSuspended())
 	{
-		FScopeLock LockAccessSenderArray(&AccessSenderArrayCriticalSection);
 		DMXSenderArray = Protocol->RegisterOutputPort(SharedThis(this));
 
 		if (DMXSenderArray.Num() > 0)
 		{
 			CommunicationDeterminator.SetHasValidSender(true);
+
 			return true;
 		}
 	}
 
 	CommunicationDeterminator.SetHasValidSender(false);
+
 	return false;
 }
 
 void FDMXOutputPort::Unregister()
-{
+{	
+	// Acquire locks in order as they're declared in the header
+	const FScopeLock LockAccessSenderArray(&AccessDMXSenderArrayMutex);
+	const FScopeLock AccesssCommunicationDeterminator(&AccessCommunicationDeterminatorMutex);
+
 	if (IsRegistered())
 	{
 		if (Protocol.IsValid())
@@ -892,7 +920,6 @@ void FDMXOutputPort::Unregister()
 			Protocol->UnregisterOutputPort(SharedThis(this));
 		}
 
-		FScopeLock LockAccessSenderArray(&AccessSenderArrayCriticalSection);
 		DMXSenderArray.Reset();
 	}
 
@@ -901,6 +928,9 @@ void FDMXOutputPort::Unregister()
 
 void FDMXOutputPort::OnSetSendDMXEnabled(bool bEnabled)
 {
+	checkf(IsInGameThread(), TEXT("Only the game thread can enable and disable send DMX."));
+
+	const FScopeLock AccesssCommunicationDeterminator(&AccessCommunicationDeterminatorMutex);
 	CommunicationDeterminator.SetSendEnabled(bEnabled);
 
 	FDMXOutputPortConfig Config = MakeOutputPortConfig();
@@ -909,6 +939,9 @@ void FDMXOutputPort::OnSetSendDMXEnabled(bool bEnabled)
 
 void FDMXOutputPort::OnSetReceiveDMXEnabled(bool bEnabled)
 {
+	checkf(IsInGameThread(), TEXT("Only the game thread can enable and disable send DMX."));
+
+	const FScopeLock AccesssCommunicationDeterminator(&AccessCommunicationDeterminatorMutex);
 	CommunicationDeterminator.SetReceiveEnabled(bEnabled);
 
 	FDMXOutputPortConfig Config = MakeOutputPortConfig();
@@ -922,11 +955,11 @@ bool FDMXOutputPort::Init()
 
 uint32 FDMXOutputPort::Run()
 {
-	const UDMXProtocolSettings* DMXSettings = GetDefault<UDMXProtocolSettings>();
+	UDMXProtocolSettings* DMXSettings = GetMutableDefault<UDMXProtocolSettings>();
 	check(DMXSettings);
-	
-	// Fixed rate delta time
-	const double SendDeltaTime = 1.f / DMXSettings->SendingRefreshRate;
+
+	SendDMXEvent = FPlatformProcess::GetSynchEventFromPool();
+	check(SendDMXEvent != nullptr);
 
 	while (!bStopping)
 	{
@@ -935,18 +968,14 @@ uint32 FDMXOutputPort::Run()
 		ProcessSendDMX();
 
 		const double EndTime = FPlatformTime::Seconds();
-		const double WaitTime = SendDeltaTime - (EndTime - StartTime);
+		const double WaitTimeMs = ((1.0 / DMXSettings->SendingRefreshRate) - (EndTime - StartTime)) * 1000.0;
 
-		if (WaitTime > 0.f)
-		{
-			// Sleep by the amount which is set in refresh rate
-			FPlatformProcess::SleepNoStats(WaitTime);
-		}
-
-		// In the unlikely case we took to long to send, we instantly continue, but do not take 
-		// further measures to compensate - We would have to run faster than DMX send rate to catch up.
+		SendDMXEvent->Wait(WaitTimeMs);
 	}
 
+	FPlatformProcess::ReturnSynchEventToPool(SendDMXEvent);
+	SendDMXEvent = nullptr;
+
 	return 0;
 }
 
@@ -974,49 +1003,55 @@ void FDMXOutputPort::ProcessSendDMX()
 	// Delay signals
 	const double Now = FPlatformTime::Seconds();
 
-	FScopeLock LockClearBuffers(&ClearBuffersCriticalSection);
-	
+	// Locking Overhead here is minimal as the locks are only in place while changing port properties.
+	// Acquire locks in order as they're declared in the header. 
+	const FScopeLock LockDMXSenderArray(&AccessDMXSenderArrayMutex);
+	const FScopeLock LockSignalFragments(&AccessSignalFragmentsMutex);
+	const FScopeLock LockExternUniverseToLatestSignalMap_PortThread(&AccessExternUniverseToLatestSignalMap_PortThreadMutex);
+	const FScopeLock LockRawListeners(&AccessRawListenersMutex);
+
 	// Write dmx fragments
+	for (;;)
 	{
-		for (;;)
+		TSharedPtr<FDMXSignalFragment, ESPMode::ThreadSafe> OldestFragment;
+		if (SignalFragments.Peek(OldestFragment))
 		{
-			TSharedPtr<FDMXSignalFragment, ESPMode::ThreadSafe> OldestFragment;
-			if (SignalFragments.Peek(OldestFragment))
+			if (OldestFragment->SendTime <= Now)
 			{
-				if (OldestFragment->SendTime <= Now)
-				{
-					const FDMXSignalSharedPtr& Signal = ExternUniverseToLatestSignalMap_PortThread.FindOrAdd(OldestFragment->ExternUniverseID, MakeShared<FDMXSignal, ESPMode::ThreadSafe>());
+				const FDMXSignalSharedPtr& Signal = ExternUniverseToLatestSignalMap_PortThread.FindOrAdd(OldestFragment->ExternUniverseID, MakeShared<FDMXSignal, ESPMode::ThreadSafe>());
 
-					// Write the fragment & meta data 
-					for (const TTuple<int32, uint8>& ChannelValueKvp : OldestFragment->ChannelToValueMap)
+				// Write the fragment & meta data 
+				for (const TTuple<int32, uint8>& ChannelValueKvp : OldestFragment->ChannelToValueMap)
+				{
+					int32 ChannelIndex = ChannelValueKvp.Key - 1;
+					// Filter invalid indicies so we can send bp calls here without testing them first.
+					if (Signal->ChannelData.IsValidIndex(ChannelIndex))
 					{
-						int32 ChannelIndex = ChannelValueKvp.Key - 1;
-						// Filter invalid indicies so we can send bp calls here without testing them first.
-						if (Signal->ChannelData.IsValidIndex(ChannelIndex))
-						{
-							Signal->ChannelData[ChannelIndex] = ChannelValueKvp.Value;
-						}
+						Signal->ChannelData[ChannelIndex] = ChannelValueKvp.Value;
 					}
+				}
 
-					Signal->ExternUniverseID = OldestFragment->ExternUniverseID;
-					Signal->Timestamp = Now;
-					Signal->Priority = Priority;
+				Signal->ExternUniverseID = OldestFragment->ExternUniverseID;
+				Signal->Timestamp = Now;
+				Signal->Priority = Priority;
 
-					// Drop the written fragment
-					SignalFragments.Pop();
+				// Drop the written fragment
+				SignalFragments.Pop();
 
-					continue;
-				}
-
-				break;
+				continue;
 			}
 
 			break;
 		}
+
+		break;
 	}
 
 	// Send new and alive DMX Signals
+	AccessCommunicationDeterminatorMutex.Lock();
 	const bool bNeedsSendDMX = CommunicationDeterminator.NeedsSendDMX();
+	AccessCommunicationDeterminatorMutex.Unlock();
+
 	for (const TTuple<int32, FDMXSignalSharedPtr>& UniverseToSignalPair : ExternUniverseToLatestSignalMap_PortThread)
 	{
 		if (UniverseToSignalPair.Value->Timestamp <= Now)
@@ -1028,7 +1063,6 @@ void FDMXOutputPort::ProcessSendDMX()
 			// Send via the protocol's sender
 			if (bNeedsSendDMX)
 			{
-				FScopeLock LockAccessSenderArray(&AccessSenderArrayCriticalSection);
 				for (const TSharedPtr<IDMXSender>& DMXSender : DMXSenderArray)
 				{
 					DMXSender->SendDMXSignal(UniverseToSignalPair.Value.ToSharedRef());
diff --git a/Engine/Plugins/VirtualProduction/DMX/DMXProtocol/Source/DMXProtocol/Public/IO/DMXInputPort.h b/Engine/Plugins/VirtualProduction/DMX/DMXProtocol/Source/DMXProtocol/Public/IO/DMXInputPort.h
--- a/Engine/Plugins/VirtualProduction/DMX/DMXProtocol/Source/DMXProtocol/Public/IO/DMXInputPort.h
+++ b/Engine/Plugins/VirtualProduction/DMX/DMXProtocol/Source/DMXProtocol/Public/IO/DMXInputPort.h
@@ -9,6 +9,7 @@
 #include "CoreMinimal.h"
 #include "Tickable.h"
 #include "Containers/Queue.h" 
+#include "HAL/CriticalSection.h"
 #include "Templates/SharedPointer.h"
 
 struct FDMXInputPortConfig;
@@ -142,6 +143,9 @@ private:
 	/** Map of raw isteners */
 	TSet<TSharedRef<FDMXRawListener>> RawListeners;
 
+	/** Mutex access for the RawListeners array */
+	FCriticalSection AccessRawListenersMutex;
+
 	/** True if the port is registered with it its protocol */
 	bool bRegistered = false;
 
diff --git a/Engine/Plugins/VirtualProduction/DMX/DMXProtocol/Source/DMXProtocol/Public/IO/DMXOutputPort.h b/Engine/Plugins/VirtualProduction/DMX/DMXProtocol/Source/DMXProtocol/Public/IO/DMXOutputPort.h
--- a/Engine/Plugins/VirtualProduction/DMX/DMXProtocol/Source/DMXProtocol/Public/IO/DMXOutputPort.h
+++ b/Engine/Plugins/VirtualProduction/DMX/DMXProtocol/Source/DMXProtocol/Public/IO/DMXOutputPort.h
@@ -20,6 +20,8 @@ class FDMXPortManager;
 class FDMXSignal;
 class IDMXSender;
 
+class FEvent;
+
 
 /** Helper to determine how DMX should be communicated (loopback, send) */
 struct FDMXOutputPortCommunicationDeterminator
@@ -93,7 +95,9 @@ struct DMXPROTOCOL_API FDMXSignalFragment
  *
  * To loopback outputs refer to DMXRawListener and DMXTickedUniverseListener.
  *
- * Can only be constructed via DMXPortManger, see FDMXPortManager::CreateOutputPort and FDMXPortManager::CreateOutputPortFromConfig
+ * Can only be constructed via DMXPortManger, see FDMXPortManager::CreateOutputPort and FDMXPortManager::CreateOutputPortFromConfig.
+ * 
+ * Note, internally locks have to be acquired in order as they're declared in the header.
  */
 class DMXPROTOCOL_API FDMXOutputPort
 	: public FDMXPort
@@ -134,9 +138,6 @@ public:
 	/** Clears all buffers */
 	void ClearBuffers();
 
-	/** Returns true if the port loopsback to engine */
-	bool IsLoopbackToEngine() const;
-
 	/** 
 	 * Game-Thread only: Gets the last signal received in specified local universe. 
 	 * 
@@ -154,10 +155,6 @@ public:
 	UE_DEPRECATED(4.27, "Use GameThreadGetDMXSignal instead. GameThreadGetDMXSignalFromRemoteUniverse only exists to support deprecated blueprint nodes.")
 	bool GameThreadGetDMXSignalFromRemoteUniverse(FDMXSignalSharedPtr& OutDMXSignal, int32 RemoteUniverseID, bool bEvenIfNotLoopbackToEngine);
 
-	/** DEPRECATED 5.0 */
-	UE_DEPRECATED(5.0, "Output Ports now support many destination addresses. Please use FDMXOutputPort::GetDestinationAddresses instead.")
-	FString GetDestinationAddress() const;
-
 	/** Returns the output port's delay in seconds */
 	FORCEINLINE double GetDelaySeconds() const { return DelaySeconds; }
 
@@ -198,20 +195,35 @@ private:
 	/** The DMX senders in use */
 	TArray<TSharedPtr<IDMXSender>> DMXSenderArray;
 
+	/** Mutex access for the RawListeners array */
+	mutable FCriticalSection AccessDMXSenderArrayMutex;
+
 	/** Buffer of the signals that are to be sent in the next frame */
 	TQueue<TSharedPtr<FDMXSignalFragment, ESPMode::ThreadSafe>> SignalFragments;
 
+	/** Mutex access for the SignalFragments Queue */
+	mutable FCriticalSection AccessSignalFragmentsMutex;
+
 	/** Map that holds the latest Signal per Universe on the Game Thread */
 	TMap<int32, FDMXSignalSharedPtr> ExternUniverseToLatestSignalMap_GameThread;
 
 	/** Map that holds the latest Signal per Universe on the Port Thread */
 	TMap<int32, FDMXSignalSharedPtr> ExternUniverseToLatestSignalMap_PortThread;
 
+	/** Mutex access for the ExternUniverseToLatestSignalMap_PortThread map */
+	mutable FCriticalSection AccessExternUniverseToLatestSignalMap_PortThreadMutex;
+
 	/** The Destination Address to send to, can be irrelevant, e.g. for art-net broadcast */
 	TArray<FDMXOutputPortDestinationAddress> DestinationAddresses;
 
+	/** Mutex access for the DestinationAddresses  */
+	mutable FCriticalSection AccessDestinationAddressesMutex;
+
 	/** Helper to determine how dmx should be communicated (loopback, send) */
 	FDMXOutputPortCommunicationDeterminator CommunicationDeterminator;
+	
+	/** Mutex access for the CommunicationDeterminator */
+	mutable FCriticalSection AccessCommunicationDeterminatorMutex;
 
 	/** Priority on which packets are being sent */
 	int32 Priority = 0;
@@ -219,6 +231,9 @@ private:
 	/** Map of raw Inputs */
 	TSet<TSharedRef<FDMXRawListener>> RawListeners;
 
+	/** Mutex access for the RawListeners array */
+	mutable FCriticalSection AccessRawListenersMutex;
+
 	/** True if the port is registered with it its protocol */
 	bool bRegistered = false;
 
@@ -231,13 +246,10 @@ private:
 	/** The frame rate of the delay */
 	FFrameRate DelayFrameRate;
 
-	/** Critical section required to be used when clearing buffers */
-	FCriticalSection AccessSenderArrayCriticalSection;
-
-	/** Critical section required to be used when clearing buffers */
-	FCriticalSection ClearBuffersCriticalSection;
+	/** Send DMX sync event */
+	FEvent* SendDMXEvent = nullptr;
 
-	/** Holds the thread object. */
+	/** Holds this thread object */
 	FRunnableThread* Thread = nullptr;
 
 	/** Flag indicating that the thread is stopping. */
diff --git a/Engine/Plugins/VirtualProduction/DMX/DMXProtocol/Source/DMXProtocolArtNet/Private/DMXProtocolArtNetReceiver.cpp b/Engine/Plugins/VirtualProduction/DMX/DMXProtocol/Source/DMXProtocolArtNet/Private/DMXProtocolArtNetReceiver.cpp
--- a/Engine/Plugins/VirtualProduction/DMX/DMXProtocol/Source/DMXProtocolArtNet/Private/DMXProtocolArtNetReceiver.cpp
+++ b/Engine/Plugins/VirtualProduction/DMX/DMXProtocol/Source/DMXProtocolArtNet/Private/DMXProtocolArtNetReceiver.cpp
@@ -96,7 +96,7 @@ void FDMXProtocolArtNetReceiver::AssignInputPort(const TSharedPtr<FDMXInputPort,
 {
 	check(!AssignedInputPorts.Contains(InputPort));
 
-	const FScopeLock ChangeAssignedInputPortsLock(&ChangeAssignedInputPortsCriticalSection);
+	const FScopeLock LockInputPorts(&AccessInputPortsMutex);
 	AssignedInputPorts.Add(InputPort);
 }
 
@@ -104,7 +104,7 @@ void FDMXProtocolArtNetReceiver::UnassignInputPort(const TSharedPtr<FDMXInputPor
 {
 	check(AssignedInputPorts.Contains(InputPort));
 
-	const FScopeLock ChangeAssignedInputPortsLock(&ChangeAssignedInputPortsCriticalSection);
+	const FScopeLock LockInputPorts(&AccessInputPortsMutex);
 	AssignedInputPorts.Remove(InputPort);
 }
 
@@ -151,8 +151,6 @@ void FDMXProtocolArtNetReceiver::Update(const FTimespan& SocketWaitTime)
 		return;
 	}
 
-	const FScopeLock ChangeAssignedInputPortsLock(&ChangeAssignedInputPortsCriticalSection);
-
 	uint32 Size = 0;
 	int32 Read = 0;
 	while (Socket->HasPendingData(Size))
@@ -240,10 +238,13 @@ void FDMXProtocolArtNetReceiver::HandleDataPacket(const TSharedRef<FArrayReader>
 
 		constexpr int32 Priority = 0;
 		FDMXSignalSharedRef DMXSignal = MakeShared<FDMXSignal, ESPMode::ThreadSafe>(FPlatformTime::Seconds(), UniverseID, Priority, TArray<uint8>(ArtNetDMXPacket.Data, DMX_UNIVERSE_SIZE));
+
+		AccessInputPortsMutex.Lock();
 		for (const TSharedPtr<FDMXInputPort, ESPMode::ThreadSafe>& InputPort : AssignedInputPorts)
 		{				
 			InputPort->InputDMXSignal(DMXSignal);
 		}
+		AccessInputPortsMutex.Unlock();
 
 		INC_DWORD_STAT(STAT_ArtNetPackagesReceived);
 	}
diff --git a/Engine/Plugins/VirtualProduction/DMX/DMXProtocol/Source/DMXProtocolArtNet/Private/DMXProtocolArtNetReceiver.h b/Engine/Plugins/VirtualProduction/DMX/DMXProtocol/Source/DMXProtocolArtNet/Private/DMXProtocolArtNetReceiver.h
--- a/Engine/Plugins/VirtualProduction/DMX/DMXProtocol/Source/DMXProtocolArtNet/Private/DMXProtocolArtNetReceiver.h
+++ b/Engine/Plugins/VirtualProduction/DMX/DMXProtocol/Source/DMXProtocolArtNet/Private/DMXProtocolArtNetReceiver.h
@@ -62,6 +62,9 @@ private:
 	/** The input ports the receiver uses */
 	TSet<TSharedPtr<FDMXInputPort, ESPMode::ThreadSafe>> AssignedInputPorts;
 
+	/** Mutex access to the input ports array */
+	FCriticalSection AccessInputPortsMutex;
+
 protected:
 	//~ Begin FRunnable implementation
 	virtual bool Init() override;
@@ -113,9 +116,6 @@ private:
 	/** The sender when packets were received */
 	TSharedPtr<FInternetAddr> ReceivedSenderInternetAddr;
 
-	/** Critical section to be used when assigned input ports are changed */
-	FCriticalSection ChangeAssignedInputPortsCriticalSection;
-
 	/** Flag indicating that the thread is stopping. */
 	TAtomic<bool> bStopping;
 
diff --git a/Engine/Plugins/VirtualProduction/DMX/DMXProtocol/Source/DMXProtocolSACN/Private/DMXProtocolSACNReceiver.cpp b/Engine/Plugins/VirtualProduction/DMX/DMXProtocol/Source/DMXProtocolSACN/Private/DMXProtocolSACNReceiver.cpp
--- a/Engine/Plugins/VirtualProduction/DMX/DMXProtocol/Source/DMXProtocolSACN/Private/DMXProtocolSACNReceiver.cpp
+++ b/Engine/Plugins/VirtualProduction/DMX/DMXProtocol/Source/DMXProtocolSACN/Private/DMXProtocolSACNReceiver.cpp
@@ -110,7 +110,7 @@ void FDMXProtocolSACNReceiver::AssignInputPort(const TSharedPtr<FDMXInputPort, E
 {
 	check(!AssignedInputPorts.Contains(InputPort));
 
-	const FScopeLock ChangeAssignedInputPortsLock(&ChangeAssignedInputPortsCriticalSection);
+	const FScopeLock LockInputPorts(&AccessInputPortsMutex);
 	AssignedInputPorts.Add(InputPort);
 
 	// Join the multicast groups for the ports where needed
@@ -143,7 +143,7 @@ void FDMXProtocolSACNReceiver::UnassignInputPort(const TSharedPtr<FDMXInputPort,
 {
 	check(AssignedInputPorts.Contains(InputPort));
 
-	const FScopeLock ChangeAssignedInputPortsLock(&ChangeAssignedInputPortsCriticalSection);
+	const FScopeLock LockInputPorts(&AccessInputPortsMutex);
 	AssignedInputPorts.Remove(InputPort);
 
 	// Leave multicast groups outside of remaining port's universe ranges
@@ -213,8 +213,6 @@ void FDMXProtocolSACNReceiver::Update(const FTimespan& SocketWaitTime)
 		return;
 	}
 
-	const FScopeLock ChangeAssignedInputPortsLock(&ChangeAssignedInputPortsCriticalSection);
-
 	uint32 Size = 0;
 	int32 NumBytesRead = 0;
 	while (Socket->HasPendingData(Size))
@@ -351,6 +349,8 @@ void FDMXProtocolSACNReceiver::HandleDataPacket(uint16 UniverseID, const TShared
 	FMemory::Memcpy(UniverseData->GetData() + IncomingDMXDMPLayer.FirstPropertyAddress, IncomingDMXDMPLayer.DMX.GetData(), IncomingDMXDMPLayer.PropertyValueCount - 1);
 	
 	FDMXSignalSharedRef DMXSignal = MakeShared<FDMXSignal, ESPMode::ThreadSafe>(FPlatformTime::Seconds(), UniverseID, IncomingDMXFramingLayer.Priority, PropertiesCacheValues[UniverseID]);
+	
+	AccessInputPortsMutex.Lock();
 	for (const TSharedPtr<FDMXInputPort, ESPMode::ThreadSafe>& InputPort : AssignedInputPorts)
 	{
 		// Check packet priority
@@ -360,6 +360,7 @@ void FDMXProtocolSACNReceiver::HandleDataPacket(uint16 UniverseID, const TShared
 		}
 		InputPort->InputDMXSignal(DMXSignal);
 	}
+	AccessInputPortsMutex.Unlock();
 
 	INC_DWORD_STAT(STAT_SACNPackagesReceived);
 }
diff --git a/Engine/Plugins/VirtualProduction/DMX/DMXProtocol/Source/DMXProtocolSACN/Private/DMXProtocolSACNReceiver.h b/Engine/Plugins/VirtualProduction/DMX/DMXProtocol/Source/DMXProtocolSACN/Private/DMXProtocolSACNReceiver.h
--- a/Engine/Plugins/VirtualProduction/DMX/DMXProtocol/Source/DMXProtocolSACN/Private/DMXProtocolSACNReceiver.h
+++ b/Engine/Plugins/VirtualProduction/DMX/DMXProtocol/Source/DMXProtocolSACN/Private/DMXProtocolSACNReceiver.h
@@ -66,6 +66,9 @@ private:
 	/** The input ports the receiver uses */
 	TSet<TSharedPtr<FDMXInputPort, ESPMode::ThreadSafe>> AssignedInputPorts;
 
+	/** Mutex access to the input ports array */
+	FCriticalSection AccessInputPortsMutex;
+
 protected:
 	//~ Begin FRunnable implementation
 	virtual bool Init() override;
@@ -113,9 +116,6 @@ private:
 	/** The destination when packets are received */
 	TSharedPtr<FInternetAddr> ReceivedDestinationInternetAddr;
 
-	/** Critical section to be used when assigned input ports are changed */
-	FCriticalSection ChangeAssignedInputPortsCriticalSection;
-
 	/** Flag indicating that the thread is stopping. */
 	TAtomic<bool> bStopping;
 
