DMX - Systemically prevent from deadlocks in DMX output ports

diff --git a/Engine/Plugins/VirtualProduction/DMX/DMXProtocol/Source/DMXProtocol/Private/IO/DMXOutputPort.cpp b/Engine/Plugins/VirtualProduction/DMX/DMXProtocol/Source/DMXProtocol/Private/IO/DMXOutputPort.cpp
--- a/Engine/Plugins/VirtualProduction/DMX/DMXProtocol/Source/DMXProtocol/Private/IO/DMXOutputPort.cpp
+++ b/Engine/Plugins/VirtualProduction/DMX/DMXProtocol/Source/DMXProtocol/Private/IO/DMXOutputPort.cpp
@@ -577,9 +577,7 @@ FDMXOutputPortConfig FDMXOutputPort::MakeOutputPortConfig() const
 {
 	checkf(IsInGameThread(), TEXT("Only game-thread can read port configs."));
 
-	// Acquire locks in order as they're declared in the header
-	const FScopeLock LockDestinationAddresses(&AccessDestinationAddressesMutex);
-	const FScopeLock LockCommunicationDeterminator(&AccessCommunicationDeterminatorMutex);
+	const FScopeLock Lock(&SendDMXMutex);
 
 	FDMXOutputPortConfigParams Params;
 	Params.PortName = PortName;
@@ -604,9 +602,7 @@ void FDMXOutputPort::UpdateFromConfig(FDMXOutputPortConfig& InOutOutputPortConfi
 {	
 	checkf(IsInGameThread(), TEXT("Only game-thread can update a port from a config."));
 
-	// Acquire locks in order as they're declared in the header
-	const FScopeLock LockDestinationAddresses(&AccessDestinationAddressesMutex);
-	const FScopeLock LockCommunicationDeterminator(&AccessCommunicationDeterminatorMutex);
+	const FScopeLock Lock(&SendDMXMutex);
 
 	// Need a valid config for the port
 	InOutOutputPortConfig.MakeValid();
@@ -703,10 +699,7 @@ void FDMXOutputPort::ClearBuffers()
 	// Otherwise the Send methods would need a lock on enqueuing, deterring the lock-free implementation.
 	checkf(IsInGameThread(), TEXT("Only the game-thread can ever clear buffers."));
 
-	// Acquire locks in order as they're declared in the header
-	const FScopeLock LockSignalFragments(&AccessSignalFragmentsMutex);
-	const FScopeLock LockExternUniverseToLatestSignalMap_PortThread(&AccessExternUniverseToLatestSignalMap_PortThreadMutex);
-	const FScopeLock LockRawlisteners(&AccessRawListenersMutex);
+	const FScopeLock Lock(&SendDMXMutex);
 
 	ExternUniverseToLatestSignalMap_GameThread.Reset();
 	ExternUniverseToLatestSignalMap_PortThread.Reset();
@@ -721,7 +714,7 @@ void FDMXOutputPort::ClearBuffers()
 
 TArray<FString> FDMXOutputPort::GetDestinationAddresses() const
 {
-	const FScopeLock LockDestinationAddresses(&AccessDestinationAddressesMutex);
+	const FScopeLock Lock(&SendDMXMutex);
 
 	TArray<FString> Result;
 	Result.Reserve(DestinationAddresses.Num());
@@ -737,7 +730,6 @@ bool FDMXOutputPort::GameThreadGetDMXSignal(int32 LocalUniverseID, FDMXSignalSha
 {
 	check(IsInGameThread());
 
-	// We don't mind to lock the communication determinator, it is ok if it is out of sync while changing.
 	const bool bNeedsLoopbackToEngine = CommunicationDeterminator.NeedsLoopbackToEngine();
 	if (bNeedsLoopbackToEngine || bEvenIfNotLoopbackToEngine)
 	{
@@ -759,7 +751,6 @@ bool FDMXOutputPort::GameThreadGetDMXSignalFromRemoteUniverse(FDMXSignalSharedPt
 	// DEPRECATED 4.27
 	check(IsInGameThread());
 
-	// We don't mind to lock the communication determinator, it is ok if it is out of sync while changing
 	const bool bNeedsLoopbackToEngine = CommunicationDeterminator.NeedsLoopbackToEngine();
 	if (bNeedsLoopbackToEngine || bEvenIfNotLoopbackToEngine)
 	{
@@ -776,30 +767,21 @@ bool FDMXOutputPort::GameThreadGetDMXSignalFromRemoteUniverse(FDMXSignalSharedPt
 
 bool FDMXOutputPort::IsRegistered() const
 {	
-	const FScopeLock LockSenderArray(&AccessDMXSenderArrayMutex);
-	if (DMXSenderArray.Num() > 0)
-	{
-		return true;
-	}
-
-	return false;
+	return bRegistered;
 }
 
 void FDMXOutputPort::AddRawListener(TSharedRef<FDMXRawListener> InRawListener)
 {
-	const FScopeLock LockRawListeners(&AccessRawListenersMutex);
+	const FScopeLock Lock(&SendDMXMutex);
 
 	check(!RawListeners.Contains(InRawListener));
 
-	// Needs to run in the game thread
-	check(IsInGameThread());
-
 	RawListeners.Add(InRawListener);
 }
 
 void FDMXOutputPort::RemoveRawListener(TSharedRef<FDMXRawListener> InRawListenerToRemove)
 {
-	const FScopeLock LockRawListeners(&AccessRawListenersMutex);
+	const FScopeLock Lock(&SendDMXMutex);
 
 	RawListeners.Remove(InRawListenerToRemove);
 }
@@ -850,7 +832,6 @@ void FDMXOutputPort::SendDMXToRemoteUniverse(const TMap<int32, uint8>& ChannelTo
 
 	if (IsExternUniverseInPortRange(RemoteUniverse))
 	{
-		// We don't mind to lock the communication determinator, it is ok if it is true while changing
 		const bool bNeedsSendDMX = CommunicationDeterminator.NeedsSendDMX();
 		const bool bNeedsLoopbackToEngine = CommunicationDeterminator.NeedsLoopbackToEngine();
 
@@ -886,9 +867,9 @@ void FDMXOutputPort::SendDMXToRemoteUniverse(const TMap<int32, uint8>& ChannelTo
 
 bool FDMXOutputPort::Register()
 {
-	// Acquire locks in order as they're declared in the header
-	const FScopeLock LockAccessSenderArray(&AccessDMXSenderArrayMutex);
-	const FScopeLock LockCommunicationDeterminator(&AccessCommunicationDeterminatorMutex);
+	checkf(IsInGameThread(), TEXT("Only Game Thread can register output ports."));
+
+	const FScopeLock Lock(&SendDMXMutex);
 
 	if (Protocol.IsValid() && IsValidPortSlow() && CommunicationDeterminator.IsSendDMXEnabled() && !FDMXPortManager::Get().AreProtocolsSuspended())
 	{
@@ -897,21 +878,23 @@ bool FDMXOutputPort::Register()
 		if (DMXSenderArray.Num() > 0)
 		{
 			CommunicationDeterminator.SetHasValidSender(true);
+			bRegistered = true;
 
 			return true;
 		}
 	}
 
 	CommunicationDeterminator.SetHasValidSender(false);
+	bRegistered = false;
 
 	return false;
 }
 
 void FDMXOutputPort::Unregister()
-{	
-	// Acquire locks in order as they're declared in the header
-	const FScopeLock LockAccessSenderArray(&AccessDMXSenderArrayMutex);
-	const FScopeLock AccesssCommunicationDeterminator(&AccessCommunicationDeterminatorMutex);
+{
+	checkf(IsInGameThread(), TEXT("Only Game Thread can unregister output ports."));
+
+	const FScopeLock Lock(&SendDMXMutex);
 
 	if (IsRegistered())
 	{
@@ -924,13 +907,13 @@ void FDMXOutputPort::Unregister()
 	}
 
 	CommunicationDeterminator.SetHasValidSender(false);
+	bRegistered = false;
 }
 
 void FDMXOutputPort::OnSetSendDMXEnabled(bool bEnabled)
 {
 	checkf(IsInGameThread(), TEXT("Only the game thread can enable and disable send DMX."));
 
-	const FScopeLock AccesssCommunicationDeterminator(&AccessCommunicationDeterminatorMutex);
 	CommunicationDeterminator.SetSendEnabled(bEnabled);
 
 	FDMXOutputPortConfig Config = MakeOutputPortConfig();
@@ -941,7 +924,6 @@ void FDMXOutputPort::OnSetReceiveDMXEnabled(bool bEnabled)
 {
 	checkf(IsInGameThread(), TEXT("Only the game thread can enable and disable send DMX."));
 
-	const FScopeLock AccesssCommunicationDeterminator(&AccessCommunicationDeterminatorMutex);
 	CommunicationDeterminator.SetReceiveEnabled(bEnabled);
 
 	FDMXOutputPortConfig Config = MakeOutputPortConfig();
@@ -1003,12 +985,8 @@ void FDMXOutputPort::ProcessSendDMX()
 	// Delay signals
 	const double Now = FPlatformTime::Seconds();
 
-	// Locking Overhead here is minimal as the locks are only in place while changing port properties.
-	// Acquire locks in order as they're declared in the header. 
-	const FScopeLock LockDMXSenderArray(&AccessDMXSenderArrayMutex);
-	const FScopeLock LockSignalFragments(&AccessSignalFragmentsMutex);
-	const FScopeLock LockExternUniverseToLatestSignalMap_PortThread(&AccessExternUniverseToLatestSignalMap_PortThreadMutex);
-	const FScopeLock LockRawListeners(&AccessRawListenersMutex);
+	// Locking Overhead here is minimal as this is only locked while changing port properties.
+	const FScopeLock Lock(&SendDMXMutex);
 
 	// Write dmx fragments
 	for (;;)
@@ -1048,9 +1026,7 @@ void FDMXOutputPort::ProcessSendDMX()
 	}
 
 	// Send new and alive DMX Signals
-	AccessCommunicationDeterminatorMutex.Lock();
 	const bool bNeedsSendDMX = CommunicationDeterminator.NeedsSendDMX();
-	AccessCommunicationDeterminatorMutex.Unlock();
 
 	for (const TTuple<int32, FDMXSignalSharedPtr>& UniverseToSignalPair : ExternUniverseToLatestSignalMap_PortThread)
 	{
diff --git a/Engine/Plugins/VirtualProduction/DMX/DMXProtocol/Source/DMXProtocol/Public/IO/DMXInputPort.h b/Engine/Plugins/VirtualProduction/DMX/DMXProtocol/Source/DMXProtocol/Public/IO/DMXInputPort.h
--- a/Engine/Plugins/VirtualProduction/DMX/DMXProtocol/Source/DMXProtocol/Public/IO/DMXInputPort.h
+++ b/Engine/Plugins/VirtualProduction/DMX/DMXProtocol/Source/DMXProtocol/Public/IO/DMXInputPort.h
@@ -10,6 +10,7 @@
 #include "Tickable.h"
 #include "Containers/Queue.h" 
 #include "HAL/CriticalSection.h"
+#include "Templates/Atomic.h"
 #include "Templates/SharedPointer.h"
 
 struct FDMXInputPortConfig;
@@ -147,7 +148,7 @@ private:
 	FCriticalSection AccessRawListenersMutex;
 
 	/** True if the port is registered with it its protocol */
-	bool bRegistered = false;
+	TAtomic<bool> bRegistered = false;
 
 	/** The unique identifier of this port, shared with the port config this was constructed from. Should not be changed after construction. */
 	FGuid PortGuid;
diff --git a/Engine/Plugins/VirtualProduction/DMX/DMXProtocol/Source/DMXProtocol/Public/IO/DMXOutputPort.h b/Engine/Plugins/VirtualProduction/DMX/DMXProtocol/Source/DMXProtocol/Public/IO/DMXOutputPort.h
--- a/Engine/Plugins/VirtualProduction/DMX/DMXProtocol/Source/DMXProtocol/Public/IO/DMXOutputPort.h
+++ b/Engine/Plugins/VirtualProduction/DMX/DMXProtocol/Source/DMXProtocol/Public/IO/DMXOutputPort.h
@@ -58,10 +58,10 @@ struct FDMXOutputPortCommunicationDeterminator
 	FORCEINLINE bool IsLoopbackToEngineEnabled() const { return bLoopbackToEngine; }
 
 private:
-	bool bLoopbackToEngine;
-	bool bReceiveEnabled;
-	bool bSendEnabled;
-	bool bHasValidSender;
+	TAtomic<bool> bLoopbackToEngine;
+	TAtomic<bool> bReceiveEnabled;
+	TAtomic<bool> bSendEnabled;
+	TAtomic<bool> bHasValidSender;
 };
 
 
@@ -195,35 +195,23 @@ private:
 	/** The DMX senders in use */
 	TArray<TSharedPtr<IDMXSender>> DMXSenderArray;
 
-	/** Mutex access for the RawListeners array */
-	mutable FCriticalSection AccessDMXSenderArrayMutex;
-
 	/** Buffer of the signals that are to be sent in the next frame */
 	TQueue<TSharedPtr<FDMXSignalFragment, ESPMode::ThreadSafe>> SignalFragments;
 
-	/** Mutex access for the SignalFragments Queue */
-	mutable FCriticalSection AccessSignalFragmentsMutex;
-
 	/** Map that holds the latest Signal per Universe on the Game Thread */
 	TMap<int32, FDMXSignalSharedPtr> ExternUniverseToLatestSignalMap_GameThread;
 
 	/** Map that holds the latest Signal per Universe on the Port Thread */
 	TMap<int32, FDMXSignalSharedPtr> ExternUniverseToLatestSignalMap_PortThread;
 
-	/** Mutex access for the ExternUniverseToLatestSignalMap_PortThread map */
-	mutable FCriticalSection AccessExternUniverseToLatestSignalMap_PortThreadMutex;
-
 	/** The Destination Address to send to, can be irrelevant, e.g. for art-net broadcast */
 	TArray<FDMXOutputPortDestinationAddress> DestinationAddresses;
 
-	/** Mutex access for the DestinationAddresses  */
-	mutable FCriticalSection AccessDestinationAddressesMutex;
-
 	/** Helper to determine how dmx should be communicated (loopback, send) */
 	FDMXOutputPortCommunicationDeterminator CommunicationDeterminator;
-	
-	/** Mutex access for the CommunicationDeterminator */
-	mutable FCriticalSection AccessCommunicationDeterminatorMutex;
+
+	/** Mutex access to members required to send DMX  */
+	mutable FCriticalSection SendDMXMutex;
 
 	/** Priority on which packets are being sent */
 	int32 Priority = 0;
@@ -231,11 +219,8 @@ private:
 	/** Map of raw Inputs */
 	TSet<TSharedRef<FDMXRawListener>> RawListeners;
 
-	/** Mutex access for the RawListeners array */
-	mutable FCriticalSection AccessRawListenersMutex;
-
-	/** True if the port is registered with it its protocol */
-	bool bRegistered = false;
+	/** True if the port is registered with its protocol */
+	TAtomic<bool> bRegistered = false;
 
 	/** The unique identifier of this port, shared with the port config this was constructed from. Should not be changed after construction. */
 	FGuid PortGuid;
