Fixed static analysis warnings exposed by the log macro refactor

#rnx

diff --git a/Engine/Plugins/VirtualProduction/DMX/DMXProtocol/Source/DMXProtocolArtNet/Private/DMXProtocolArtNetSender.cpp b/Engine/Plugins/VirtualProduction/DMX/DMXProtocol/Source/DMXProtocolArtNet/Private/DMXProtocolArtNetSender.cpp
--- a/Engine/Plugins/VirtualProduction/DMX/DMXProtocol/Source/DMXProtocolArtNet/Private/DMXProtocolArtNetSender.cpp
+++ b/Engine/Plugins/VirtualProduction/DMX/DMXProtocol/Source/DMXProtocolArtNet/Private/DMXProtocolArtNetSender.cpp
@@ -186,7 +186,7 @@ void FDMXProtocolArtNetSender::SendDMXSignal(const FDMXSignalSharedRef& DMXSigna
 			ISocketSubsystem* SocketSubsystem = ISocketSubsystem::Get(PLATFORM_SOCKETSUBSYSTEM);
 			TEnumAsByte<ESocketErrors> RecvFromError = SocketSubsystem->GetLastErrorCode();
 
-			UE_LOG(LogDMXProtocol, Error, TEXT("Failed send DMX to %s with Error Code %d"), *DestinationInternetAddr->ToString(false), RecvFromError);
+			UE_LOG(LogDMXProtocol, Error, TEXT("Failed send DMX to %s with Error Code %d"), *DestinationInternetAddr->ToString(false), RecvFromError.GetValue());
 
 			bErrorEverLogged = true;
 		}
diff --git a/Engine/Plugins/VirtualProduction/DMX/DMXProtocol/Source/DMXProtocolSACN/Private/DMXProtocolSACNSender.cpp b/Engine/Plugins/VirtualProduction/DMX/DMXProtocol/Source/DMXProtocolSACN/Private/DMXProtocolSACNSender.cpp
--- a/Engine/Plugins/VirtualProduction/DMX/DMXProtocol/Source/DMXProtocolSACN/Private/DMXProtocolSACNSender.cpp
+++ b/Engine/Plugins/VirtualProduction/DMX/DMXProtocol/Source/DMXProtocolSACN/Private/DMXProtocolSACNSender.cpp
@@ -217,7 +217,7 @@ void FDMXProtocolSACNSender::SendDMXSignal(const FDMXSignalSharedRef& DMXSignal)
 			ISocketSubsystem* SocketSubsystem = ISocketSubsystem::Get(PLATFORM_SOCKETSUBSYSTEM);
 			TEnumAsByte<ESocketErrors> RecvFromError = SocketSubsystem->GetLastErrorCode();
 
-			UE_LOG(LogDMXProtocol, Error, TEXT("Failed send DMX to %s with Error Code %d"), *DestinationInternetAddr->ToString(false), RecvFromError);
+			UE_LOG(LogDMXProtocol, Error, TEXT("Failed send DMX to %s with Error Code %d"), *DestinationInternetAddr->ToString(false), RecvFromError.GetValue());
 
 			bErrorEverLogged = true;
 		}
