DMX - Fix sACN uses any instead of the specified NIC to send DMX

diff --git a/Engine/Plugins/VirtualProduction/DMX/DMXProtocol/Source/DMXProtocol/Private/IO/DMXPortManager.cpp b/Engine/Plugins/VirtualProduction/DMX/DMXProtocol/Source/DMXProtocol/Private/IO/DMXPortManager.cpp
--- a/Engine/Plugins/VirtualProduction/DMX/DMXProtocol/Source/DMXProtocol/Private/IO/DMXPortManager.cpp
+++ b/Engine/Plugins/VirtualProduction/DMX/DMXProtocol/Source/DMXProtocol/Private/IO/DMXPortManager.cpp
@@ -17,6 +17,19 @@ TUniquePtr<FDMXPortManager> FDMXPortManager::CurrentManager;
 FDMXPortManager::~FDMXPortManager()
 {
 	checkf(!CurrentManager.IsValid(), TEXT("ShutdownManager was not called"));
+
+	// Other objects may still hold shared pointers to ports up until engine shuts down.
+	// Hence unregister all ports when the engine and by that this manager shut down. 
+	// Alike the resource of ports (e.g. sockets) are always properly released.
+	for (const FDMXInputPortSharedRef& InputPort : InputPorts)
+	{
+		InputPort->Unregister();
+	}
+
+	for (const FDMXOutputPortSharedRef& OutputPort : OutputPorts)
+	{
+		OutputPort->Unregister();
+	}
 }
 
 FDMXPortManager& FDMXPortManager::Get()
diff --git a/Engine/Plugins/VirtualProduction/DMX/DMXProtocol/Source/DMXProtocolSACN/Private/DMXProtocolSACNSender.cpp b/Engine/Plugins/VirtualProduction/DMX/DMXProtocol/Source/DMXProtocolSACN/Private/DMXProtocolSACNSender.cpp
--- a/Engine/Plugins/VirtualProduction/DMX/DMXProtocol/Source/DMXProtocolSACN/Private/DMXProtocolSACNSender.cpp
+++ b/Engine/Plugins/VirtualProduction/DMX/DMXProtocol/Source/DMXProtocolSACN/Private/DMXProtocolSACNSender.cpp
@@ -68,7 +68,7 @@ TSharedPtr<FDMXProtocolSACNSender> FDMXProtocolSACNSender::TryCreateUnicastSende
 		return nullptr;
 	}
 
-	FIPv4Endpoint NewNetworkInterfaceEndpoint = FIPv4Endpoint(NewNetworkInterfaceInternetAddr);
+	const FIPv4Endpoint NewNetworkInterfaceEndpoint = FIPv4Endpoint(NewNetworkInterfaceInternetAddr);
 	
 	FSocket* NewSocket =
 		FUdpSocketBuilder(TEXT("UDPSACNUnicastSocket"))
@@ -108,12 +108,14 @@ TSharedPtr<FDMXProtocolSACNSender> FDMXProtocolSACNSender::TryCreateMulticastSen
 		UE_LOG(LogDMXProtocol, Error, TEXT("Cannot create sACN sender: Invalid IP address: %s"), *InNetworkInterfaceIP);
 		return nullptr;
 	}
+	FIPv4Endpoint NewNetworkInterfaceEndpoint = FIPv4Endpoint(NewNetworkInterfaceInternetAddr);
 
 	FSocket* NewSocket = 
 		FUdpSocketBuilder(TEXT("UDPSACNMulticastSocket"))
         .AsBlocking()
         .AsReusable()
         .WithMulticastLoopback()
+		.WithMulticastInterface(NewNetworkInterfaceEndpoint.Address)
         .WithMulticastTtl(1);
 	
 	if(!NewSocket)
