DMX - Fix DMX Fixture Types Project Setting Attribute Keywords cannot use spaces

diff --git a/Engine/Plugins/VirtualProduction/DMX/DMXProtocol/Source/DMXProtocol/Private/DMXAttribute.cpp b/Engine/Plugins/VirtualProduction/DMX/DMXProtocol/Source/DMXProtocol/Private/DMXAttribute.cpp
--- a/Engine/Plugins/VirtualProduction/DMX/DMXProtocol/Source/DMXProtocol/Private/DMXAttribute.cpp
+++ b/Engine/Plugins/VirtualProduction/DMX/DMXProtocol/Source/DMXProtocol/Private/DMXAttribute.cpp
@@ -114,7 +114,6 @@ void FDMXAttribute::CleanupKeywords()
 	// support tabs too
 	Keywords = Keywords.ConvertTabsToSpaces(1);
 	Keywords.TrimStartAndEndInline();
-	Keywords = Keywords.Replace(TEXT(" "), TEXT(","));
 	TArray<FString> KeywordsArray;
 	Keywords.ParseIntoArray(KeywordsArray, TEXT(","), true);
 	Keywords = FString::Join(KeywordsArray, TEXT(", "));
