DMX - sACN sender now correctly sets UE5 as source name

diff --git a/Engine/Plugins/VirtualProduction/DMX/DMXProtocol/Source/DMXProtocolSACN/Private/Packets/DMXProtocolE131PDUPacket.h b/Engine/Plugins/VirtualProduction/DMX/DMXProtocol/Source/DMXProtocolSACN/Private/Packets/DMXProtocolE131PDUPacket.h
--- a/Engine/Plugins/VirtualProduction/DMX/DMXProtocol/Source/DMXProtocolSACN/Private/Packets/DMXProtocolE131PDUPacket.h
+++ b/Engine/Plugins/VirtualProduction/DMX/DMXProtocol/Source/DMXProtocolSACN/Private/Packets/DMXProtocolE131PDUPacket.h
@@ -38,7 +38,7 @@ public:
 public:
 	uint16 FlagsAndLength = ACN_DMX_FRAMING_FLAGS_AND_LENGTH_SIZE;
 	uint32 Vector = VECTOR_E131_DATA_PACKET;
-	TArray<uint8, TInlineAllocator<ACN_SOURCE_NAME_SIZE>> SourceName = { 'U', 'E', '4', '\0' };
+	TArray<uint8, TInlineAllocator<ACN_SOURCE_NAME_SIZE>> SourceName = { 'U', 'E', '5', '\0' };
 	uint8 Priority = 1;
 	uint16 SynchronizationAddress = 0;
 	uint8 SequenceNumber = 0;
