DMX - Fix output port threads get stalled when send rate cannot be met

diff --git a/Engine/Plugins/VirtualProduction/DMX/DMXProtocol/Source/DMXProtocol/Private/IO/DMXOutputPort.cpp b/Engine/Plugins/VirtualProduction/DMX/DMXProtocol/Source/DMXProtocol/Private/IO/DMXOutputPort.cpp
--- a/Engine/Plugins/VirtualProduction/DMX/DMXProtocol/Source/DMXProtocol/Private/IO/DMXOutputPort.cpp
+++ b/Engine/Plugins/VirtualProduction/DMX/DMXProtocol/Source/DMXProtocol/Private/IO/DMXOutputPort.cpp
@@ -14,6 +14,7 @@
 #include "HAL/Event.h"
 #include "HAL/IConsoleManager.h"
 #include "HAL/RunnableThread.h"
+#include <limits.h>
 #include "Misc/FrameRate.h"
 
 
@@ -948,8 +949,11 @@ uint32 FDMXOutputPort::Run()
 
 		const double EndTime = FPlatformTime::Seconds();
 		const double WaitTimeMs = ((1.0 / DMXSettings->SendingRefreshRate) - (EndTime - StartTime)) * 1000.0;
-
-		SendDMXEvent->Wait(WaitTimeMs);
+		
+		if (WaitTimeMs > 0.0 && WaitTimeMs < std::numeric_limits<uint32>::max())
+		{
+			SendDMXEvent->Wait(WaitTimeMs);
+		}
 	}
 
 	FPlatformProcess::ReturnSynchEventToPool(SendDMXEvent);
@@ -961,6 +965,12 @@ uint32 FDMXOutputPort::Run()
 void FDMXOutputPort::Stop()
 {
 	bStopping = true;
+
+	// Effectively trigger the stop
+	if (SendDMXEvent)
+	{
+		SendDMXEvent->Trigger();
+	}
 }
 
 void FDMXOutputPort::Exit()
