DMX - Fix log level for invalid ports prevents from packaging with invalid ports

diff --git a/Engine/Plugins/VirtualProduction/DMX/DMXProtocol/Source/DMXProtocolArtNet/Private/DMXProtocolArtNetReceiver.cpp b/Engine/Plugins/VirtualProduction/DMX/DMXProtocol/Source/DMXProtocolArtNet/Private/DMXProtocolArtNetReceiver.cpp
--- a/Engine/Plugins/VirtualProduction/DMX/DMXProtocol/Source/DMXProtocolArtNet/Private/DMXProtocolArtNetReceiver.cpp
+++ b/Engine/Plugins/VirtualProduction/DMX/DMXProtocol/Source/DMXProtocolArtNet/Private/DMXProtocolArtNetReceiver.cpp
@@ -59,7 +59,7 @@ TSharedPtr<FDMXProtocolArtNetReceiver> FDMXProtocolArtNetReceiver::TryCreate(con
 	TSharedPtr<FInternetAddr> EndpointInternetAddr = FDMXProtocolUtils::CreateInternetAddr(IPAddress, ARTNET_PORT);
 	if (!EndpointInternetAddr.IsValid())
 	{
-		UE_LOG(LogDMXProtocol, Error, TEXT("Cannot create Art-Net receiver: Invalid IP address: %s"), *IPAddress);
+		UE_LOG(LogDMXProtocol, Warning, TEXT("Cannot create Art-Net receiver: Invalid IP address: %s"), *IPAddress);
 		return nullptr;
 	}
 
@@ -72,7 +72,7 @@ TSharedPtr<FDMXProtocolArtNetReceiver> FDMXProtocolArtNetReceiver::TryCreate(con
 
 	if (!NewListeningSocket)
 	{
-		UE_LOG(LogDMXProtocol, Error, TEXT("Cannot create Art-Net receiver: Error create ListeningSocket for: %s"), *IPAddress);
+		UE_LOG(LogDMXProtocol, Warning, TEXT("Cannot create Art-Net receiver: Error create ListeningSocket for: %s"), *IPAddress);
 		return nullptr;
 	}
 
diff --git a/Engine/Plugins/VirtualProduction/DMX/DMXProtocol/Source/DMXProtocolArtNet/Private/DMXProtocolArtNetSender.cpp b/Engine/Plugins/VirtualProduction/DMX/DMXProtocol/Source/DMXProtocolArtNet/Private/DMXProtocolArtNetSender.cpp
--- a/Engine/Plugins/VirtualProduction/DMX/DMXProtocol/Source/DMXProtocolArtNet/Private/DMXProtocolArtNetSender.cpp
+++ b/Engine/Plugins/VirtualProduction/DMX/DMXProtocol/Source/DMXProtocolArtNet/Private/DMXProtocolArtNetSender.cpp
@@ -54,7 +54,7 @@ TSharedPtr<FDMXProtocolArtNetSender> FDMXProtocolArtNetSender::TryCreateUnicastS
 	TSharedPtr<FInternetAddr> NewNetworkInterfaceInternetAddr = FDMXProtocolUtils::CreateInternetAddr(InNetworkInterfaceIP, ARTNET_SENDER_PORT);
 	if (!NewNetworkInterfaceInternetAddr.IsValid())
 	{
-		UE_LOG(LogDMXProtocol, Error, TEXT("Cannot create Art-Net sender: Invalid IP address: %s"), *InNetworkInterfaceIP);
+		UE_LOG(LogDMXProtocol, Warning, TEXT("Cannot create Art-Net sender: Invalid IP address: %s"), *InNetworkInterfaceIP);
 		return nullptr;
 	}
 
@@ -68,7 +68,7 @@ TSharedPtr<FDMXProtocolArtNetSender> FDMXProtocolArtNetSender::TryCreateUnicastS
 	
 	if(!NewSocket)
 	{
-		UE_LOG(LogDMXProtocol, Error, TEXT("Invalid Network Interface IP %s for DMX Port. Please update your Output Port in Project Settings -> Plugins -> DMX Plugin"), *InNetworkInterfaceIP);
+		UE_LOG(LogDMXProtocol, Warning, TEXT("Invalid Network Interface IP %s for DMX Port. Please update your Output Port in Project Settings -> Plugins -> DMX Plugin"), *InNetworkInterfaceIP);
 		return nullptr;
 	}
 
@@ -76,7 +76,7 @@ TSharedPtr<FDMXProtocolArtNetSender> FDMXProtocolArtNetSender::TryCreateUnicastS
 	TSharedPtr<FInternetAddr> NewUnicastInternetAddr = FDMXProtocolUtils::CreateInternetAddr(InUnicastIP, ARTNET_PORT);
 	if (!NewUnicastInternetAddr.IsValid())
 	{
-		UE_LOG(LogDMXProtocol, Error, TEXT("Invalid Unicast IP %s for DMX Port. Please update your Output Port in Project Settings -> Plugins -> DMX Plugin"), *InUnicastIP);
+		UE_LOG(LogDMXProtocol, Warning, TEXT("Invalid Unicast IP %s for DMX Port. Please update your Output Port in Project Settings -> Plugins -> DMX Plugin"), *InUnicastIP);
 		return nullptr;
 	}
 
@@ -94,7 +94,7 @@ TSharedPtr<FDMXProtocolArtNetSender> FDMXProtocolArtNetSender::TryCreateBroadcas
 	TSharedPtr<FInternetAddr> NewNetworkInterfaceInternetAddr = FDMXProtocolUtils::CreateInternetAddr(InNetworkInterfaceIP, ARTNET_SENDER_PORT);
 	if (!NewNetworkInterfaceInternetAddr.IsValid())
 	{
-		UE_LOG(LogDMXProtocol, Error, TEXT("Cannot create Art-Net sender: Invalid IP address: %s"), *InNetworkInterfaceIP);
+		UE_LOG(LogDMXProtocol, Warning, TEXT("Cannot create Art-Net sender: Invalid IP address: %s"), *InNetworkInterfaceIP);
 		return nullptr;
 	}
 
@@ -109,7 +109,7 @@ TSharedPtr<FDMXProtocolArtNetSender> FDMXProtocolArtNetSender::TryCreateBroadcas
 
 	if(!NewSocket)
 	{
-		UE_LOG(LogDMXProtocol, Error, TEXT("Invalid Network Interface IP %s for DMX Port. Please update your Output Ports in Project Settings -> Plugins -> DMX Plugin"), *InNetworkInterfaceIP);
+		UE_LOG(LogDMXProtocol, Warning, TEXT("Invalid Network Interface IP %s for DMX Port. Please update your Output Ports in Project Settings -> Plugins -> DMX Plugin"), *InNetworkInterfaceIP);
 		return nullptr;
 	}
 
@@ -186,7 +186,7 @@ void FDMXProtocolArtNetSender::SendDMXSignal(const FDMXSignalSharedRef& DMXSigna
 			ISocketSubsystem* SocketSubsystem = ISocketSubsystem::Get(PLATFORM_SOCKETSUBSYSTEM);
 			TEnumAsByte<ESocketErrors> RecvFromError = SocketSubsystem->GetLastErrorCode();
 
-			UE_LOG(LogDMXProtocol, Error, TEXT("Failed send DMX to %s with Error Code %d"), *DestinationInternetAddr->ToString(false), RecvFromError.GetValue());
+			UE_LOG(LogDMXProtocol, Error, TEXT("Socket subsystem failed sending DMX to %s with Error Code %d"), *DestinationInternetAddr->ToString(false), RecvFromError.GetValue());
 
 			bErrorEverLogged = true;
 		}
diff --git a/Engine/Plugins/VirtualProduction/DMX/DMXProtocol/Source/DMXProtocolSACN/Private/DMXProtocolSACNReceiver.cpp b/Engine/Plugins/VirtualProduction/DMX/DMXProtocol/Source/DMXProtocolSACN/Private/DMXProtocolSACNReceiver.cpp
--- a/Engine/Plugins/VirtualProduction/DMX/DMXProtocol/Source/DMXProtocolSACN/Private/DMXProtocolSACNReceiver.cpp
+++ b/Engine/Plugins/VirtualProduction/DMX/DMXProtocol/Source/DMXProtocolSACN/Private/DMXProtocolSACNReceiver.cpp
@@ -61,7 +61,7 @@ TSharedPtr<FDMXProtocolSACNReceiver> FDMXProtocolSACNReceiver::TryCreate(const T
 	const TSharedPtr<FInternetAddr> EndpointInternetAddr = FDMXProtocolUtils::CreateInternetAddr(IPAddress, ACN_PORT);
 	if (!EndpointInternetAddr.IsValid())
 	{
-		UE_LOG(LogDMXProtocol, Error, TEXT("Cannot create sACN receiver: Invalid IP address: %s"), *IPAddress);
+		UE_LOG(LogDMXProtocol, Warning, TEXT("Cannot create sACN receiver: Invalid IP address: %s"), *IPAddress);
 		return nullptr;
 	}
 	const FIPv4Endpoint Endpoint(EndpointInternetAddr);
@@ -77,7 +77,7 @@ TSharedPtr<FDMXProtocolSACNReceiver> FDMXProtocolSACNReceiver::TryCreate(const T
 
 	if (!NewListeningSocket)
 	{
-		UE_LOG(LogDMXProtocol, Error, TEXT("Cannot create sACN receiver: Error create ListeningSocket for: %s"), *IPAddress);
+		UE_LOG(LogDMXProtocol, Warning, TEXT("Cannot create sACN receiver: Error create ListeningSocket for: %s"), *IPAddress);
 		return nullptr;
 	}
 
@@ -86,7 +86,7 @@ TSharedPtr<FDMXProtocolSACNReceiver> FDMXProtocolSACNReceiver::TryCreate(const T
 		ISocketSubsystem* SocketSubsystem = ISocketSubsystem::Get(PLATFORM_SOCKETSUBSYSTEM);
 		SocketSubsystem->DestroySocket(NewListeningSocket);
 
-		UE_LOG(LogDMXProtocol, Error, TEXT("Cannot create sACN receiver: Platform does not support IP_PKTINFO required for sACN."), *IPAddress);
+		UE_LOG(LogDMXProtocol, Warning, TEXT("Cannot create sACN receiver: Platform does not support IP_PKTINFO required for sACN."), *IPAddress);
 		return nullptr;
 	}
 
diff --git a/Engine/Plugins/VirtualProduction/DMX/DMXProtocol/Source/DMXProtocolSACN/Private/DMXProtocolSACNSender.cpp b/Engine/Plugins/VirtualProduction/DMX/DMXProtocol/Source/DMXProtocolSACN/Private/DMXProtocolSACNSender.cpp
--- a/Engine/Plugins/VirtualProduction/DMX/DMXProtocol/Source/DMXProtocolSACN/Private/DMXProtocolSACNSender.cpp
+++ b/Engine/Plugins/VirtualProduction/DMX/DMXProtocol/Source/DMXProtocolSACN/Private/DMXProtocolSACNSender.cpp
@@ -64,7 +64,7 @@ TSharedPtr<FDMXProtocolSACNSender> FDMXProtocolSACNSender::TryCreateUnicastSende
 	TSharedPtr<FInternetAddr> NewNetworkInterfaceInternetAddr = FDMXProtocolUtils::CreateInternetAddr(InNetworkInterfaceIP, ACN_PORT);
 	if (!NewNetworkInterfaceInternetAddr.IsValid())
 	{
-		UE_LOG(LogDMXProtocol, Error, TEXT("Cannot create sACN sender: Invalid IP address: %s"), *InNetworkInterfaceIP);
+		UE_LOG(LogDMXProtocol, Warning, TEXT("Cannot create sACN sender: Invalid IP address: %s"), *InNetworkInterfaceIP);
 		return nullptr;
 	}
 
@@ -79,7 +79,7 @@ TSharedPtr<FDMXProtocolSACNSender> FDMXProtocolSACNSender::TryCreateUnicastSende
 
 	if(!NewSocket)
 	{
-		UE_LOG(LogDMXProtocol, Error, TEXT("Invalid Network Interface IP %s for DMX Port. Please update your Output Port in Project Settings -> Plugins -> DMX Plugin"), *InNetworkInterfaceIP);
+		UE_LOG(LogDMXProtocol, Warning, TEXT("Invalid Network Interface IP %s for DMX Port. Please update your Output Port in Project Settings -> Plugins -> DMX Plugin"), *InNetworkInterfaceIP);
 		return nullptr;
 	}
 
@@ -87,7 +87,7 @@ TSharedPtr<FDMXProtocolSACNSender> FDMXProtocolSACNSender::TryCreateUnicastSende
 	TSharedPtr<FInternetAddr> NewUnicastInternetAddr = FDMXProtocolUtils::CreateInternetAddr(InUnicastIP, ACN_PORT);
 	if (!NewUnicastInternetAddr.IsValid())
 	{
-		UE_LOG(LogDMXProtocol, Error, TEXT("Invalid Unicast IP %s for DMX Port. Please update your Output Port in Project Settings -> Plugins -> DMX Plugin"), *InUnicastIP);
+		UE_LOG(LogDMXProtocol, Warning, TEXT("Invalid Unicast IP %s for DMX Port. Please update your Output Port in Project Settings -> Plugins -> DMX Plugin"), *InUnicastIP);
 		return nullptr;
 	}
 
@@ -105,7 +105,7 @@ TSharedPtr<FDMXProtocolSACNSender> FDMXProtocolSACNSender::TryCreateMulticastSen
 	const TSharedPtr<FInternetAddr> NewNetworkInterfaceInternetAddr = FDMXProtocolUtils::CreateInternetAddr(InNetworkInterfaceIP, ACN_PORT);
 	if (!NewNetworkInterfaceInternetAddr.IsValid())
 	{
-		UE_LOG(LogDMXProtocol, Error, TEXT("Cannot create sACN sender: Invalid IP address: %s"), *InNetworkInterfaceIP);
+		UE_LOG(LogDMXProtocol, Warning, TEXT("Cannot create sACN sender: Invalid IP address: %s"), *InNetworkInterfaceIP);
 		return nullptr;
 	}
 	FIPv4Endpoint NewNetworkInterfaceEndpoint = FIPv4Endpoint(NewNetworkInterfaceInternetAddr);
@@ -120,7 +120,7 @@ TSharedPtr<FDMXProtocolSACNSender> FDMXProtocolSACNSender::TryCreateMulticastSen
 	
 	if(!NewSocket)
 	{
-		UE_LOG(LogDMXProtocol, Error, TEXT("Invalid Network Interface IP %s for DMX Port. Please update your Output Ports in Project Settings -> Plugins -> DMX Plugin"), *InNetworkInterfaceIP);
+		UE_LOG(LogDMXProtocol, Warning, TEXT("Invalid Network Interface IP %s for DMX Port. Please update your Output Ports in Project Settings -> Plugins -> DMX Plugin"), *InNetworkInterfaceIP);
 		return nullptr;
 	}
 
@@ -219,7 +219,7 @@ void FDMXProtocolSACNSender::SendDMXSignal(const FDMXSignalSharedRef& DMXSignal)
 			ISocketSubsystem* SocketSubsystem = ISocketSubsystem::Get(PLATFORM_SOCKETSUBSYSTEM);
 			TEnumAsByte<ESocketErrors> RecvFromError = SocketSubsystem->GetLastErrorCode();
 
-			UE_LOG(LogDMXProtocol, Error, TEXT("Failed send DMX to %s with Error Code %d"), *DestinationInternetAddr->ToString(false), RecvFromError.GetValue());
+			UE_LOG(LogDMXProtocol, Error, TEXT("Socket subsystem failed sending DMX to %s with Error Code %d"), *DestinationInternetAddr->ToString(false), RecvFromError.GetValue());
 
 			bErrorEverLogged = true;
 		}
