DMX - Fix a possible dead lock when changing ports while activity monitor is open

diff --git a/Engine/Plugins/VirtualProduction/DMX/DMXProtocol/Source/DMXProtocol/Private/IO/DMXRawListener.cpp b/Engine/Plugins/VirtualProduction/DMX/DMXProtocol/Source/DMXProtocol/Private/IO/DMXRawListener.cpp
--- a/Engine/Plugins/VirtualProduction/DMX/DMXProtocol/Source/DMXProtocol/Private/IO/DMXRawListener.cpp
+++ b/Engine/Plugins/VirtualProduction/DMX/DMXProtocol/Source/DMXProtocol/Private/IO/DMXRawListener.cpp
@@ -2,6 +2,8 @@
 
 #include "IO/DMXRawListener.h"
 
+#include "Async/Async.h"
+#include "Async/Future.h"
 #include "IO/DMXPort.h"
 
 
@@ -46,8 +48,6 @@ void FDMXRawListener::Stop()
 	// listener and minimal performance benefits. 
 	if (FDMXPortSharedPtr PinnedOwnerPort = OwnerPort.Pin())
 	{
-		FScopeLock Lock(&ListenerCriticalSection);
-
 		PinnedOwnerPort->RemoveRawListener(AsShared());
 	}
 
@@ -66,8 +66,6 @@ void FDMXRawListener::EnqueueSignal(void* Producer, const FDMXSignalSharedRef& S
 	checkf(PinnedOwnerPort.IsValid() && ProducerObj == PinnedOwnerPort.Get(), TEXT("More than one producer detected in FDMXRawListener::DequeueSignal."));
 #endif // UE_BUILD_DEBUG
 
-	FScopeLock Lock(&ListenerCriticalSection);
-
 	RawBuffer.Enqueue(Signal);
 }
 
@@ -82,8 +80,6 @@ bool FDMXRawListener::DequeueSignal(void* Consumer, FDMXSignalSharedPtr& OutSign
 	checkf(ConsumerObj == Consumer, TEXT("More than one consumer detected in FDMXRawListener::DequeueSignal."));
 #endif // UE_BUILD_DEBUG
 
-	FScopeLock Lock(&ListenerCriticalSection);
-
 	if (RawBuffer.Dequeue(OutSignal))
 	{
 		OutLocalUniverseID = OutSignal->ExternUniverseID - ExternUniverseOffset;
@@ -100,15 +96,27 @@ bool FDMXRawListener::DequeueSignal(void* Consumer, FDMXSignalSharedPtr& OutSign
 
 void FDMXRawListener::ClearBuffer()
 {
-	FScopeLock Lock(&ListenerCriticalSection);
-
-	RawBuffer.Empty();
+	if (IsInGameThread())
+	{
+		RawBuffer.Empty();
+	}
+	else
+	{
+		// The buffer has to be cleared on a consumer thread
+		TPromise<bool> Promise;
+		AsyncTask(ENamedThreads::GameThread, [&Promise, this]()
+			{
+				RawBuffer.Empty();
+				Promise.SetValue(true);
+			});
+
+		// Lock the calling thread until the promise was fullfilled
+		Promise.GetFuture().Get();
+	}
 }
 
 void FDMXRawListener::OnPortUpdated()
 {
-	FScopeLock Lock(&ListenerCriticalSection);
-
 	if (TSharedPtr<FDMXPort, ESPMode::ThreadSafe> PinnedOwnerPort = OwnerPort.Pin())
 	{
 		ExternUniverseOffset = PinnedOwnerPort->GetExternUniverseOffset();
diff --git a/Engine/Plugins/VirtualProduction/DMX/DMXProtocol/Source/DMXProtocol/Public/IO/DMXRawListener.h b/Engine/Plugins/VirtualProduction/DMX/DMXProtocol/Source/DMXProtocol/Public/IO/DMXRawListener.h
--- a/Engine/Plugins/VirtualProduction/DMX/DMXProtocol/Source/DMXProtocol/Public/IO/DMXRawListener.h
+++ b/Engine/Plugins/VirtualProduction/DMX/DMXProtocol/Source/DMXProtocol/Public/IO/DMXRawListener.h
@@ -4,6 +4,7 @@
 
 #include "DMXProtocolCommon.h"
 
+#include <atomic>
 #include "CoreMinimal.h"
 #include "Containers/Queue.h"
 #include "Misc/ScopeLock.h"
@@ -73,11 +74,8 @@ private:
 	/** The port that owns this Listener */
 	TWeakPtr<FDMXPort, ESPMode::ThreadSafe> OwnerPort;
 
-	/** Lock used when the buffer is cleared */
-	FCriticalSection ListenerCriticalSection;
-
 	/** The extern universe offset of the port */
-	int32 ExternUniverseOffset;
+	TAtomic<int32> ExternUniverseOffset;
 
 	/** True when stopped */
 	bool bStopped;
