DMX - Make setting an extern universe for ports in project settings more explicit

diff --git a/Engine/Plugins/VirtualProduction/DMX/DMXProtocol/Source/DMXProtocol/Private/DMXProtocolSettings.cpp b/Engine/Plugins/VirtualProduction/DMX/DMXProtocol/Source/DMXProtocol/Private/DMXProtocolSettings.cpp
--- a/Engine/Plugins/VirtualProduction/DMX/DMXProtocol/Source/DMXProtocol/Private/DMXProtocolSettings.cpp
+++ b/Engine/Plugins/VirtualProduction/DMX/DMXProtocol/Source/DMXProtocol/Private/DMXProtocolSettings.cpp
@@ -102,14 +102,9 @@ void UDMXProtocolSettings::PostInitProperties()
 	{
 		UE_LOG(LogDMXProtocol, Log, TEXT("Overridden Default Receive DMX Enabled from command line, set to %s."), bDefaultReceiveDMXEnabled ? TEXT("True") : TEXT("False"));
 	}
-	OverrideReceiveDMXEnabled(bDefaultReceiveDMXEnabled);	
-}
-
-void UDMXProtocolSettings::PostLoad()
-{
-	Super::PostLoad();
+	OverrideReceiveDMXEnabled(bDefaultReceiveDMXEnabled);
 
-	// Upgrade from single to many destination addresses for ports
+	// Upgrade from single to many destination addresses for ports for projects created before 4.27
 	const int32 CustomVersion = GetLinkerCustomVersion(FDMXProtocolObjectVersion::GUID);
 	if (CustomVersion < FDMXProtocolObjectVersion::OutputPortSupportsManyUnicastAddresses)
 	{
@@ -133,6 +128,17 @@ void UDMXProtocolSettings::PostLoad()
 			}
 		}
 	}
+
+	// Allow port configs to take their own upgrade path
+	for (FDMXInputPortConfig& InputPortConfig : InputPortConfigs)
+	{
+		InputPortConfig.Upgrade();
+	}
+
+	for (FDMXOutputPortConfig& OutputPortConfig : OutputPortConfigs)
+	{
+		OutputPortConfig.Upgrade();
+	}
 }
 
 #if WITH_EDITOR
@@ -200,45 +206,45 @@ void UDMXProtocolSettings::PostEditChangeChainProperty(FPropertyChangedChainEven
 	{
 		FDMXPortManager::Get().UpdateFromProtocolSettings();
 	}
-	else if (PropertyChangedChainEvent.ChangeType != EPropertyChangeType::Interactive)
+	else if (
+		PropertyName == GET_MEMBER_NAME_CHECKED(UDMXProtocolSettings, InputPortConfigs) ||
+		PropertyName == GET_MEMBER_NAME_CHECKED(UDMXProtocolSettings, OutputPortConfigs) ||
+		(InputPortConfigStruct && InputPortConfigStruct == PropertyOwnerStruct) ||
+		(OutputPortConfigStruct && OutputPortConfigStruct == PropertyOwnerStruct))
 	{
-		if (PropertyName == GET_MEMBER_NAME_CHECKED(UDMXProtocolSettings, InputPortConfigs) ||
-			PropertyName == GET_MEMBER_NAME_CHECKED(UDMXProtocolSettings, OutputPortConfigs) ||
-			(InputPortConfigStruct && InputPortConfigStruct == PropertyOwnerStruct) ||
-			(OutputPortConfigStruct && OutputPortConfigStruct == PropertyOwnerStruct))
+		// When duplicating configs, the guid will be duplicated, so we have to create unique ones instead
+		if (PropertyChangedChainEvent.ChangeType == EPropertyChangeType::Duplicate)
 		{
-			if (PropertyChangedChainEvent.ChangeType == EPropertyChangeType::Duplicate)
+			int32 ChangedIndex = PropertyChangedChainEvent.GetArrayIndex(PropertyName.ToString());
+			if (ensureAlways(ChangedIndex != INDEX_NONE))
 			{
-				// When duplicating configs, the guid will be duplicated, so we have to create unique ones instead
-
-				int32 ChangedIndex = PropertyChangedChainEvent.GetArrayIndex(PropertyName.ToString());
-				if (ensureAlways(ChangedIndex != INDEX_NONE))
+				if (PropertyName == GET_MEMBER_NAME_CHECKED(UDMXProtocolSettings, InputPortConfigs))
 				{
-					if (PropertyName == GET_MEMBER_NAME_CHECKED(UDMXProtocolSettings, InputPortConfigs))
+					const int32 IndexOfDuplicate = InputPortConfigs.FindLastByPredicate([this, ChangedIndex](const FDMXInputPortConfig& InputPortConfig) {
+						return InputPortConfigs[ChangedIndex].GetPortGuid() == InputPortConfig.GetPortGuid();
+						});
+
+					if (ensureAlways(IndexOfDuplicate != ChangedIndex))
 					{
-						const int32 IndexOfDuplicate = InputPortConfigs.FindLastByPredicate([this, ChangedIndex](const FDMXInputPortConfig& InputPortConfig) {
-							return InputPortConfigs[ChangedIndex].GetPortGuid() == InputPortConfig.GetPortGuid();
-							});
-
-						if (ensureAlways(IndexOfDuplicate != ChangedIndex))
-						{
-							InputPortConfigs[IndexOfDuplicate] = FDMXInputPortConfig(FGuid::NewGuid(), InputPortConfigs[ChangedIndex]);
-						}
+						InputPortConfigs[IndexOfDuplicate] = FDMXInputPortConfig(FGuid::NewGuid(), InputPortConfigs[ChangedIndex]);
 					}
-					else if (PropertyName == GET_MEMBER_NAME_CHECKED(UDMXProtocolSettings, OutputPortConfigs))
+				}
+				else if (PropertyName == GET_MEMBER_NAME_CHECKED(UDMXProtocolSettings, OutputPortConfigs))
+				{
+					const int32 IndexOfDuplicate = OutputPortConfigs.FindLastByPredicate([this, ChangedIndex](const FDMXOutputPortConfig& OutputPortConfig) {
+						return OutputPortConfigs[ChangedIndex].GetPortGuid() == OutputPortConfig.GetPortGuid();
+						});
+
+					if (ensureAlways(IndexOfDuplicate != ChangedIndex))
 					{
-						const int32 IndexOfDuplicate = OutputPortConfigs.FindLastByPredicate([this, ChangedIndex](const FDMXOutputPortConfig& OutputPortConfig) {
-							return OutputPortConfigs[ChangedIndex].GetPortGuid() == OutputPortConfig.GetPortGuid();
-							});
-
-						if (ensureAlways(IndexOfDuplicate != ChangedIndex))
-						{
-							OutputPortConfigs[IndexOfDuplicate] = FDMXOutputPortConfig(FGuid::NewGuid(), OutputPortConfigs[ChangedIndex]);
-						}
+						OutputPortConfigs[IndexOfDuplicate] = FDMXOutputPortConfig(FGuid::NewGuid(), OutputPortConfigs[ChangedIndex]);
 					}
 				}
 			}
+		}
 
+		if (PropertyChangedChainEvent.ChangeType != EPropertyChangeType::Interactive)
+		{
 			constexpr bool bForceUpdateRegistrationWithProtocol = true;
 			FDMXPortManager::Get().UpdateFromProtocolSettings(bForceUpdateRegistrationWithProtocol);
 		}
@@ -248,28 +254,28 @@ void UDMXProtocolSettings::PostEditChangeChainProperty(FPropertyChangedChainEven
 }
 #endif // WITH_EDITOR
 
-void UDMXProtocolSettings::OverrideSendDMXEnabled(bool bEnabled) 
+void UDMXProtocolSettings::OverrideSendDMXEnabled(bool bEnabled)
 {
-	bOverrideSendDMXEnabled = bEnabled; 
-	
+	bOverrideSendDMXEnabled = bEnabled;
+
 	OnSetSendDMXEnabledDelegate.Broadcast(bEnabled);
 
 	// OnSetSendDMXEnabled is deprecated 5.1 and can be removed in a future release
-PRAGMA_DISABLE_DEPRECATION_WARNINGS
+	PRAGMA_DISABLE_DEPRECATION_WARNINGS
 	OnSetSendDMXEnabled.Broadcast(bEnabled);
-PRAGMA_ENABLE_DEPRECATION_WARNINGS
+	PRAGMA_ENABLE_DEPRECATION_WARNINGS
 }
 
-void UDMXProtocolSettings::OverrideReceiveDMXEnabled(bool bEnabled) 
-{ 
-	bOverrideReceiveDMXEnabled = bEnabled; 
+void UDMXProtocolSettings::OverrideReceiveDMXEnabled(bool bEnabled)
+{
+	bOverrideReceiveDMXEnabled = bEnabled;
 
 	OnSetReceiveDMXEnabledDelegate.Broadcast(bEnabled);
 
 	// OnSetReceiveDMXEnabled is deprecated 5.1 and can be removed in a future release
-PRAGMA_DISABLE_DEPRECATION_WARNINGS
+	PRAGMA_DISABLE_DEPRECATION_WARNINGS
 	OnSetReceiveDMXEnabled.Broadcast(bEnabled);
-PRAGMA_ENABLE_DEPRECATION_WARNINGS
+	PRAGMA_ENABLE_DEPRECATION_WARNINGS
 }
 
 FString UDMXProtocolSettings::GetUniqueInputPortName() const
diff --git a/Engine/Plugins/VirtualProduction/DMX/DMXProtocol/Source/DMXProtocol/Private/IO/DMXInputPortConfig.cpp b/Engine/Plugins/VirtualProduction/DMX/DMXProtocol/Source/DMXProtocol/Private/IO/DMXInputPortConfig.cpp
--- a/Engine/Plugins/VirtualProduction/DMX/DMXProtocol/Source/DMXProtocol/Private/IO/DMXInputPortConfig.cpp
+++ b/Engine/Plugins/VirtualProduction/DMX/DMXProtocol/Source/DMXProtocol/Private/IO/DMXInputPortConfig.cpp
@@ -126,6 +126,14 @@ void FDMXInputPortConfig::MakeValid()
 	}
 }
 
+void FDMXInputPortConfig::Upgrade()
+{
+#if WITH_EDITOR
+	// Initialize the bIsExternUnivereStartEditable property that was added 5.3. Also mends manually edited project ini files.
+	bIsExternUnivereStartEditable = ExternUniverseStart != LocalUniverseStart;
+#endif
+}
+
 FString FDMXInputPortConfig::GetDeviceAddress() const
 {
 	// Return the Command Line Device Address if it is set
diff --git a/Engine/Plugins/VirtualProduction/DMX/DMXProtocol/Source/DMXProtocol/Private/IO/DMXOutputPortConfig.cpp b/Engine/Plugins/VirtualProduction/DMX/DMXProtocol/Source/DMXProtocol/Private/IO/DMXOutputPortConfig.cpp
--- a/Engine/Plugins/VirtualProduction/DMX/DMXProtocol/Source/DMXProtocol/Private/IO/DMXOutputPortConfig.cpp
+++ b/Engine/Plugins/VirtualProduction/DMX/DMXProtocol/Source/DMXProtocol/Private/IO/DMXOutputPortConfig.cpp
@@ -144,6 +144,14 @@ void FDMXOutputPortConfig::MakeValid()
 	}
 }
 
+void FDMXOutputPortConfig::Upgrade()
+{
+#if WITH_EDITOR
+	// Initialize the bIsExternUnivereStartEditable property that was added 5.3. Also mends manually edited project ini files.
+	bIsExternUnivereStartEditable = ExternUniverseStart != LocalUniverseStart;
+#endif
+}
+
 FString FDMXOutputPortConfig::GetDeviceAddress() const
 {
 	// Return the Command Line Device Address if it is set
diff --git a/Engine/Plugins/VirtualProduction/DMX/DMXProtocol/Source/DMXProtocol/Public/DMXProtocolSettings.h b/Engine/Plugins/VirtualProduction/DMX/DMXProtocol/Source/DMXProtocol/Public/DMXProtocolSettings.h
--- a/Engine/Plugins/VirtualProduction/DMX/DMXProtocol/Source/DMXProtocol/Public/DMXProtocolSettings.h
+++ b/Engine/Plugins/VirtualProduction/DMX/DMXProtocol/Source/DMXProtocol/Public/DMXProtocolSettings.h
@@ -2,19 +2,17 @@
 
 #pragma once
 
+#include "CoreMinimal.h"
 #include "DMXProtocolTypes.h"
 #include "DMXAttribute.h"
 #include "IO/DMXInputPortConfig.h"
 #include "IO/DMXOutputPortConfig.h"
-
-#include "CoreMinimal.h"
 #include "Misc/Optional.h"
 #include "UObject/Object.h"
 
 #include "DMXProtocolSettings.generated.h"
 
 
-
 /**  
  * DMX Project Settings. 
  * 
@@ -36,7 +34,6 @@ public:
 
 	// ~Begin UObject Interface
 	virtual void PostInitProperties() override;
-	virtual void PostLoad() override;
 #if WITH_EDITOR
 	virtual void PostEditChangeProperty(FPropertyChangedEvent& PropertyChangedEvent) override;
 	virtual void PostEditChangeChainProperty(FPropertyChangedChainEvent& PropertyChangedChainEvent) override;
diff --git a/Engine/Plugins/VirtualProduction/DMX/DMXProtocol/Source/DMXProtocol/Public/IO/DMXInputPortConfig.h b/Engine/Plugins/VirtualProduction/DMX/DMXProtocol/Source/DMXProtocol/Public/IO/DMXInputPortConfig.h
--- a/Engine/Plugins/VirtualProduction/DMX/DMXProtocol/Source/DMXProtocol/Public/IO/DMXInputPortConfig.h
+++ b/Engine/Plugins/VirtualProduction/DMX/DMXProtocol/Source/DMXProtocol/Public/IO/DMXInputPortConfig.h
@@ -79,6 +79,9 @@ public:
 	/** Changes members to result in a valid config */
 	void MakeValid();
 
+	/** Upgrades this config to the latest engine version */
+	void Upgrade();
+
 	FORCEINLINE const FString& GetPortName() const { return PortName; }
 	FORCEINLINE const FName& GetProtocolName() const { return ProtocolName; }
 	FORCEINLINE EDMXCommunicationType GetCommunicationType() const { return CommunicationType; }
@@ -102,6 +105,7 @@ public:
 	static FName GetLocalUniverseStartPropertyNameChecked() { return GET_MEMBER_NAME_CHECKED(FDMXInputPortConfig, LocalUniverseStart); }
 	static FName GetNumUniversesPropertyNameChecked() { return GET_MEMBER_NAME_CHECKED(FDMXInputPortConfig, NumUniverses); }
 	static FName GetExternUniverseStartPropertyNameChecked() { return GET_MEMBER_NAME_CHECKED(FDMXInputPortConfig, ExternUniverseStart); }
+	static FName GetIsExternUnivereStartEditablePropertyNameChecked() { return GET_MEMBER_NAME_CHECKED(FDMXInputPortConfig, bIsExternUnivereStartEditable); }
 	static FName GetPriorityStrategyPropertyNameChecked() { return GET_MEMBER_NAME_CHECKED(FDMXInputPortConfig, PriorityStrategy); }
 	static FName GetPriorityPropertyNameChecked() { return GET_MEMBER_NAME_CHECKED(FDMXInputPortConfig, Priority); }
 	static FName GetPortGuidPropertyNameChecked() { return GET_MEMBER_NAME_CHECKED(FDMXInputPortConfig, PortGuid); }
@@ -150,12 +154,18 @@ protected:
 	int32 NumUniverses = 10;
 
 	/** 
-	 * The start address this being transposed to. 
-	 * E.g. if LocalUniverseStart is 1 and this is 100, Local Universe 1 is sent/received as Universe 100.
+	 * Remaps the range of local Universes (Num Universes from Local Universe Start) to a different set of Universes at the Protocol level.
+	 * For example, for Art-Net this allows to map local Universe 1 to extern Universe 0. 
 	 */
-	UPROPERTY(Config, BlueprintReadOnly, EditDefaultsOnly, Category = "Port Config")
+	UPROPERTY(Config, BlueprintReadOnly, EditDefaultsOnly, Category = "Port Config", Meta = (DisplayName = "Protocol Universe Remap"))
 	int32 ExternUniverseStart = 1;
 
+#if WITH_EDITORONLY_DATA
+	/** Edit condition for the ExternUniverseStart property */
+	UPROPERTY(Config, EditDefaultsOnly, Category = "Port Config")
+	bool bIsExternUnivereStartEditable = false;
+#endif
+
 	/** How to deal with the priority value */
 	UPROPERTY(Config, BlueprintReadOnly, EditDefaultsOnly, Category = "Port Config")
 	EDMXPortPriorityStrategy PriorityStrategy = EDMXPortPriorityStrategy::Highest;
diff --git a/Engine/Plugins/VirtualProduction/DMX/DMXProtocol/Source/DMXProtocol/Public/IO/DMXOutputPortConfig.h b/Engine/Plugins/VirtualProduction/DMX/DMXProtocol/Source/DMXProtocol/Public/IO/DMXOutputPortConfig.h
--- a/Engine/Plugins/VirtualProduction/DMX/DMXProtocol/Source/DMXProtocol/Public/IO/DMXOutputPortConfig.h
+++ b/Engine/Plugins/VirtualProduction/DMX/DMXProtocol/Source/DMXProtocol/Public/IO/DMXOutputPortConfig.h
@@ -85,6 +85,9 @@ public:
 	/** Changes members to result in a valid config */
 	void MakeValid();
 
+	/** Upgrades this config to the latest engine version */
+	void Upgrade();
+
 	FORCEINLINE const FString& GetPortName() const { return PortName; }
 	FORCEINLINE const FName& GetProtocolName() const { return ProtocolName; }
 	FORCEINLINE EDMXCommunicationType GetCommunicationType() const { return CommunicationType; }
@@ -115,6 +118,7 @@ public:
 	static FName GetLocalUniverseStartPropertyNameChecked() { return GET_MEMBER_NAME_CHECKED(FDMXOutputPortConfig, LocalUniverseStart); }
 	static FName GetNumUniversesPropertyNameChecked() { return GET_MEMBER_NAME_CHECKED(FDMXOutputPortConfig, NumUniverses); }
 	static FName GetExternUniverseStartPropertyNameChecked() { return GET_MEMBER_NAME_CHECKED(FDMXOutputPortConfig, ExternUniverseStart); }
+	static FName GetIsExternUnivereStartEditablePropertyNameChecked() { return GET_MEMBER_NAME_CHECKED(FDMXOutputPortConfig, bIsExternUnivereStartEditable); }
 	static FName GetPriorityPropertyNameChecked() { return GET_MEMBER_NAME_CHECKED(FDMXOutputPortConfig, Priority); }
 	static FName GetDelayPropertyNameChecked() { return GET_MEMBER_NAME_CHECKED(FDMXOutputPortConfig, Delay); }
 	static FName GetDelayFrameRatePropertyNameChecked() { return GET_MEMBER_NAME_CHECKED(FDMXOutputPortConfig, DelayFrameRate); }
@@ -178,13 +182,19 @@ protected:
 	UPROPERTY(Config, BlueprintReadOnly, EditDefaultsOnly, Category = "Port Config", Meta = (DisplayName = "Amount of Universes"))
 	int32 NumUniverses = 10;
 
-	/** 
-	 * The start address this being transposed to. 
-	 * E.g. if LocalUniverseStart is 1 and this is 100, Local Universe 1 is sent/received as Universe 100.
+	/**
+	 * Remaps the range of local Universes (Num Universes from Local Universe Start) to a different set of Universes at the Protocol level.
+	 * For example, for Art-Net this allows to map local Universe 1 to extern Universe 0.
 	 */
-	UPROPERTY(Config, BlueprintReadOnly, EditDefaultsOnly, Category = "Port Config")
+	UPROPERTY(Config, BlueprintReadOnly, EditDefaultsOnly, Category = "Port Config", Meta = (DisplayName = "Protocol Universe Remap"))
 	int32 ExternUniverseStart = 1;
 
+#if WITH_EDITORONLY_DATA
+	/** Edit condition for the ExternUniverseStart property */
+	UPROPERTY(Config, EditDefaultsOnly, Category = "Port Config")
+	bool bIsExternUnivereStartEditable = false;
+#endif
+
 	/** Priority on which packets are being sent */
 	UPROPERTY(Config, BlueprintReadOnly, EditDefaultsOnly, Category = "Port Config")
 	int32 Priority = 100;
diff --git a/Engine/Plugins/VirtualProduction/DMX/DMXProtocol/Source/DMXProtocolEditor/Private/DetailsCustomizations/DMXInputPortConfigCustomization.cpp b/Engine/Plugins/VirtualProduction/DMX/DMXProtocol/Source/DMXProtocolEditor/Private/DetailsCustomizations/DMXInputPortConfigCustomization.cpp
--- a/Engine/Plugins/VirtualProduction/DMX/DMXProtocol/Source/DMXProtocolEditor/Private/DetailsCustomizations/DMXInputPortConfigCustomization.cpp
+++ b/Engine/Plugins/VirtualProduction/DMX/DMXProtocol/Source/DMXProtocolEditor/Private/DetailsCustomizations/DMXInputPortConfigCustomization.cpp
@@ -2,28 +2,28 @@
 
 #include "DMXInputPortConfigCustomization.h"
 
+#include "DetailWidgetRow.h"
 #include "DMXProtocolModule.h"
 #include "DMXProtocolSettings.h"
 #include "DMXProtocolTypes.h"
 #include "DMXProtocolUtils.h"
+#include "IDetailChildrenBuilder.h"
+#include "IDetailPropertyRow.h" 
 #include "Interfaces/IDMXProtocol.h"
 #include "IO/DMXInputPort.h"
 #include "IO/DMXInputPortConfig.h"
 #include "IO/DMXPortManager.h"
+#include "IPropertyUtilities.h"
+#include "Misc/Guid.h" 
+#include "ScopedTransaction.h"
+#include "Styling/AppStyle.h"
+#include "Styling/CoreStyle.h"
+#include "Widgets/Input/SCheckBox.h"
 #include "Widgets/SDMXCommunicationTypeComboBox.h"
 #include "Widgets/SDMXIPAddressEditWidget.h"
 #include "Widgets/SDMXProtocolNameComboBox.h"
-
-#include "DetailWidgetRow.h"
-#include "IDetailChildrenBuilder.h"
-#include "IDetailPropertyRow.h" 
-#include "IPropertyUtilities.h"
-#include "ScopedTransaction.h"
-#include "Misc/Guid.h" 
 #include "Widgets/Text/STextBlock.h" 
 
-#include "Styling/AppStyle.h"
 
-
 #define LOCTEXT_NAMESPACE "DMXInputPortConfigCustomization"
 
@@ -63,10 +63,15 @@
 	AutoCompleteDeviceAddressEnabledHandle = PropertyHandles.FindChecked(FDMXInputPortConfig::GetAutoCompleteDeviceAddressEnabledPropertyNameChecked());
 	AutoCompleteDeviceAddressHandle = PropertyHandles.FindChecked(FDMXInputPortConfig::GetAutoCompleteDeviceAddressPropertyNameChecked());
 	DeviceAddressHandle = PropertyHandles.FindChecked(FDMXInputPortConfig::GetDeviceAddressPropertyNameChecked());
+	LocalUniverseStartHandle = PropertyHandles.FindChecked(FDMXInputPortConfig::GetLocalUniverseStartPropertyNameChecked());
+	ExternUniverseStartHandle = PropertyHandles.FindChecked(FDMXInputPortConfig::GetExternUniverseStartPropertyNameChecked());
+	IsExternUniverseStartEditableHandle = PropertyHandles.FindChecked(FDMXInputPortConfig::GetIsExternUnivereStartEditablePropertyNameChecked());
 	PortGuidHandle = PropertyHandles.FindChecked(FDMXInputPortConfig::GetPortGuidPropertyNameChecked());
 
 	// Hande property changes
 	AutoCompleteDeviceAddressHandle->SetOnPropertyValueChanged(FSimpleDelegate::CreateSP(this, &FDMXInputPortConfigCustomization::UpdateAutoCompleteDeviceAddressTextBox));
+	LocalUniverseStartHandle->SetOnPropertyValueChanged(FSimpleDelegate::CreateSP(this, &FDMXInputPortConfigCustomization::OnLocalUniverseStartChanged));
+	IsExternUniverseStartEditableHandle->SetOnPropertyValueChanged(FSimpleDelegate::CreateSP(this, &FDMXInputPortConfigCustomization::OnIsExternUniverseStartEditableChanged));
 
 	// Ports always need a valid Guid (cannot be blueprinted)
 	if (!GetPortGuid().IsValid())
@@ -91,7 +96,9 @@
 		}
 
 		// Don't add the bAutoCompeteDeviceAddressEnabled property
-		if (Iter.Value() == AutoCompleteDeviceAddressEnabledHandle)
+		// Don't add the bIsExternUniverseStartEditable property
+		if (Iter.Value() == AutoCompleteDeviceAddressEnabledHandle ||
+			Iter.Value() == IsExternUniverseStartEditableHandle)
 		{
 			continue;
 		}
@@ -115,6 +122,10 @@
 		{
 			GenerateDeviceAddressRow(PropertyRow);
 		}
+		else if (Iter.Value() == ExternUniverseStartHandle)
+		{
+			GenerateExternUniverseStartRow(PropertyRow);
+		}
 		else if (Iter.Key() == FDMXInputPortConfig::GetPriorityStrategyPropertyNameChecked())
 		{
 
@@ -203,55 +214,6 @@
 		];
 }
 
-void FDMXInputPortConfigCustomization::GenerateAutoCompleteDeviceAddressRow(IDetailPropertyRow& PropertyRow)
-{
-	// Mimic a Inline Edit Condition Toggle - Doing it directly via Meta Specifiers is not possible. -
-	// The Inline Edit Condition Toggle property wouldn't get properly serialized when being a Config property.
-	const TSharedRef<SWidget> AutoCompleteDeviceAddressEnableldPropertyValueWidget = AutoCompleteDeviceAddressEnabledHandle->CreatePropertyValueWidget();
-
-	const TSharedRef<SWidget> AutoCompleteDeviceAddressPropertyNameWidget = AutoCompleteDeviceAddressHandle->CreatePropertyNameWidget();
-	const TSharedRef<SWidget> AutoCompleteDeviceAddressPropertyValueWidget = AutoCompleteDeviceAddressHandle->CreatePropertyValueWidget();
-
-	AutoCompleteDeviceAddressPropertyNameWidget->SetEnabled(TAttribute<bool>::CreateLambda([this]()
-		{
-			return IsAutoCompleteDeviceAddressEnabled();
-		}));
-	AutoCompleteDeviceAddressPropertyValueWidget->SetEnabled(TAttribute<bool>::CreateLambda([this]()
-		{
-			return IsAutoCompleteDeviceAddressEnabled();
-		}));
-
-	PropertyRow.CustomWidget()
-		.NameContent()
-		[
-			SNew(SHorizontalBox)
-	
-			// Edit Inline Condition Toggle Checkbox
-			+ SHorizontalBox::Slot()
-			.HAlign(HAlign_Left)
-			.VAlign(VAlign_Fill)
-			.AutoWidth()
-			[
-				AutoCompleteDeviceAddressEnableldPropertyValueWidget
-			]
-
-			// Property Name
-			+ SHorizontalBox::Slot()
-			.Padding(8.f, 0.f, 0.f, 0.f)
-			.HAlign(HAlign_Left)
-			.VAlign(VAlign_Fill)
-			.AutoWidth()
-			[
-				AutoCompleteDeviceAddressPropertyNameWidget
-			]
-		]
-		.ValueContent()
-		[
-			// Property Value
-			AutoCompleteDeviceAddressPropertyValueWidget
-		];
-}
-
 void FDMXInputPortConfigCustomization::GenerateDeviceAddressRow(IDetailPropertyRow& PropertyRow)
 {
 	// Customizate the IPAddress property to show a combo box with available IP addresses
@@ -315,6 +277,55 @@
 		];
 }
 
+void FDMXInputPortConfigCustomization::GenerateAutoCompleteDeviceAddressRow(IDetailPropertyRow& PropertyRow)
+{
+	// Mimic a Inline Edit Condition Toggle - Doing it directly via Meta Specifiers is not possible. -
+	// The Inline Edit Condition Toggle property wouldn't get properly serialized when being a Config property.
+	const TSharedRef<SWidget> AutoCompleteDeviceAddressEnableldPropertyValueWidget = AutoCompleteDeviceAddressEnabledHandle->CreatePropertyValueWidget();
+
+	const TSharedRef<SWidget> AutoCompleteDeviceAddressPropertyNameWidget = AutoCompleteDeviceAddressHandle->CreatePropertyNameWidget();
+	const TSharedRef<SWidget> AutoCompleteDeviceAddressPropertyValueWidget = AutoCompleteDeviceAddressHandle->CreatePropertyValueWidget();
+
+	AutoCompleteDeviceAddressPropertyNameWidget->SetEnabled(TAttribute<bool>::CreateLambda([this]()
+		{
+			return IsAutoCompleteDeviceAddressEnabled();
+		}));
+	AutoCompleteDeviceAddressPropertyValueWidget->SetEnabled(TAttribute<bool>::CreateLambda([this]()
+		{
+			return IsAutoCompleteDeviceAddressEnabled();
+		}));
+
+	PropertyRow.CustomWidget()
+		.NameContent()
+		[
+			SNew(SHorizontalBox)
+	
+			// Edit Inline Condition Toggle Checkbox
+			+ SHorizontalBox::Slot()
+			.HAlign(HAlign_Left)
+			.VAlign(VAlign_Fill)
+			.AutoWidth()
+			[
+				AutoCompleteDeviceAddressEnableldPropertyValueWidget
+			]
+
+			// Property Name
+			+ SHorizontalBox::Slot()
+			.Padding(8.f, 0.f, 0.f, 0.f)
+			.HAlign(HAlign_Left)
+			.VAlign(VAlign_Fill)
+			.AutoWidth()
+			[
+				AutoCompleteDeviceAddressPropertyNameWidget
+			]
+		]
+		.ValueContent()
+		[
+			// Property Value
+			AutoCompleteDeviceAddressPropertyValueWidget
+		];
+}
+
 void FDMXInputPortConfigCustomization::UpdateAutoCompleteDeviceAddressTextBox()
 {
 	// Set text on the AutoCompleteDeviceAddress Text Box depending on what the resulting IP is
@@ -341,6 +352,112 @@
 	}
 }
 
+void FDMXInputPortConfigCustomization::GenerateExternUniverseStartRow(IDetailPropertyRow& PropertyRow)
+{
+	const TAttribute<bool> CanEditExternUniverseAttribute = TAttribute<bool>::CreateLambda([this]()
+		{
+			bool bIsExternUniverseStartEditable;
+			if (IsExternUniverseStartEditableHandle->GetValue(bIsExternUniverseStartEditable) == FPropertyAccess::Success)
+			{
+				return bIsExternUniverseStartEditable ? true : false;
+			}
+			return false;
+		});
+
+	TSharedPtr<SWidget> ExternUniverseStartNameWidget;
+	TSharedPtr<SWidget> ExternUniverseStartValueWidget;
+	FDetailWidgetRow ExternUniverseStartRow;
+	PropertyRow.GetDefaultWidgets(ExternUniverseStartNameWidget, ExternUniverseStartValueWidget, ExternUniverseStartRow);
+
+	const FText IsExternUniverseEnabledTooltip = LOCTEXT("IsExternUniverseEnabledWarningTooltip", "Remaps the range of local Universes (Num Universes from Local Universe Start) to a different set of Universes at the Protocol level.\nFor example, for Art-Net this allows to map local Universe 1 to extern Universe 0.");
+
+	PropertyRow.CustomWidget()
+		.NameContent()
+		[
+			SNew(SHorizontalBox)
+
+			+ SHorizontalBox::Slot()
+			.AutoWidth()
+			.Padding(0.f, 0.f, 4.f, 0.f)
+			[
+				SNew(SCheckBox)
+				.ToolTipText(IsExternUniverseEnabledTooltip)
+				.IsChecked_Lambda([CanEditExternUniverseAttribute]()
+					{
+						return CanEditExternUniverseAttribute.Get() ? ECheckBoxState::Checked : ECheckBoxState::Unchecked;
+					})
+				.OnCheckStateChanged_Lambda([this](ECheckBoxState NewState)
+					{
+						const bool bIsExternUniverseStartEditable = NewState == ECheckBoxState::Checked;
+						IsExternUniverseStartEditableHandle->SetValue(bIsExternUniverseStartEditable);
+					})
+			]
+
+			+ SHorizontalBox::Slot()
+			.AutoWidth()
+			[
+				SNew(SBox)
+				.IsEnabled(CanEditExternUniverseAttribute)
+				[
+					ExternUniverseStartNameWidget.ToSharedRef()
+				]
+			]
+		]
+		.ValueContent()
+		[
+			SNew(SHorizontalBox)
+			.IsEnabled(CanEditExternUniverseAttribute)
+
+			+ SHorizontalBox::Slot()
+			.Padding(0.f, 0.f, 4.f, 0.f)
+			[
+				ExternUniverseStartValueWidget.ToSharedRef()
+			]
+
+			+ SHorizontalBox::Slot()
+			.AutoWidth()
+			.Padding(4.f)
+			[
+				SNew(SImage)
+				.Visibility_Lambda([CanEditExternUniverseAttribute]()
+					{
+						return CanEditExternUniverseAttribute.Get() ? EVisibility::Visible : EVisibility::Collapsed;
+					})
+				.Image(FCoreStyle::Get().GetBrush("Icons.Warning"))
+				.ToolTipText(IsExternUniverseEnabledTooltip)
+			]
+		];
+}
+
+void FDMXInputPortConfigCustomization::OnLocalUniverseStartChanged()
+{
+	bool bCanEditExternUniverse;
+	int32 LocalUniverse;
+	if (IsExternUniverseStartEditableHandle->GetValue(bCanEditExternUniverse) == FPropertyAccess::Success &&
+		!bCanEditExternUniverse &&
+		LocalUniverseStartHandle->GetValue(LocalUniverse) == FPropertyAccess::Success)
+	{
+		ExternUniverseStartHandle->NotifyPreChange();
+		// Change the extern universe interactive only. Rely on the local universe start change completing the transaction.
+		ExternUniverseStartHandle->SetValue(LocalUniverse, EPropertyValueSetFlags::InteractiveChange);
+		ExternUniverseStartHandle->NotifyPostChange(EPropertyChangeType::Interactive);
+	}
+}
+
+void FDMXInputPortConfigCustomization::OnIsExternUniverseStartEditableChanged()
+{
+	bool bCanEditExternUniverse;
+	int32 LocalUniverse;
+	if (IsExternUniverseStartEditableHandle->GetValue(bCanEditExternUniverse) == FPropertyAccess::Success &&
+		!bCanEditExternUniverse &&
+		LocalUniverseStartHandle->GetValue(LocalUniverse) == FPropertyAccess::Success)
+	{
+		ExternUniverseStartHandle->NotifyPreChange();
+		ExternUniverseStartHandle->SetValue(LocalUniverse);
+		ExternUniverseStartHandle->NotifyPostChange(EPropertyChangeType::ValueSet);
+	}
+}
+
 void FDMXInputPortConfigCustomization::OnProtocolNameSelected()
 {
 	FName ProtocolName = ProtocolNameComboBox->GetSelectedProtocolName();
diff --git a/Engine/Plugins/VirtualProduction/DMX/DMXProtocol/Source/DMXProtocolEditor/Private/DetailsCustomizations/DMXInputPortConfigCustomization.h b/Engine/Plugins/VirtualProduction/DMX/DMXProtocol/Source/DMXProtocolEditor/Private/DetailsCustomizations/DMXInputPortConfigCustomization.h
--- a/Engine/Plugins/VirtualProduction/DMX/DMXProtocol/Source/DMXProtocolEditor/Private/DetailsCustomizations/DMXInputPortConfigCustomization.h
+++ b/Engine/Plugins/VirtualProduction/DMX/DMXProtocol/Source/DMXProtocolEditor/Private/DetailsCustomizations/DMXInputPortConfigCustomization.h
@@ -39,15 +39,24 @@ private:
 	/** Generates the customized Communication Type row */
 	void GenerateCommunicationTypeRow(IDetailPropertyRow& PropertyRow);
 
-	/** Generates the customized Auto Complete Device Address row */
-	void GenerateAutoCompleteDeviceAddressRow(IDetailPropertyRow& PropertyRow);
-
 	/** Generates the customized Device Address row */
 	void GenerateDeviceAddressRow(IDetailPropertyRow& PropertyRow);
 
+	/** Generates the customized Auto Complete Device Address row */
+	void GenerateAutoCompleteDeviceAddressRow(IDetailPropertyRow& PropertyRow);
+
 	/** Updates the Auto Complete Device Address Text Box */
 	void UpdateAutoCompleteDeviceAddressTextBox();
 
+	/** Generates the customized Extern Universe Start row */
+	void GenerateExternUniverseStartRow(IDetailPropertyRow& PropertyRow);
+
+	/** Called when the LocalUniverseStart property changed */
+	void OnLocalUniverseStartChanged();
+
+	/** Called when the bIsExternUniverseStartEditable property changed */
+	void OnIsExternUniverseStartEditableChanged();
+
 	/** Called when a Protocol Name was selected */
 	void OnProtocolNameSelected();
 
@@ -90,6 +99,15 @@ private:
 	/** Property handle to the CommunicationType property */
 	TSharedPtr<IPropertyHandle> CommunicationTypeHandle;
 
+	/** Property handle to the LocalUniverseStart property */
+	TSharedPtr<IPropertyHandle> LocalUniverseStartHandle;
+
+	/** Property handle to the ExternUniverseStart property */
+	TSharedPtr<IPropertyHandle> ExternUniverseStartHandle;
+
+	/** Property handle to the IsExternUniverseStartEditable property */
+	TSharedPtr<IPropertyHandle> IsExternUniverseStartEditableHandle;
+
 	/** Property handle to the PortGuid property */
 	TSharedPtr<IPropertyHandle> PortGuidHandle;
 
@@ -105,7 +123,6 @@ private:
 	/** Text Block that displays the device address when it is auto completed */
 	TSharedPtr<STextBlock> AutoCompletedDeviceAddressTextBlock;
 
-
 	/** Property utilities for this customization */
 	TSharedPtr<IPropertyUtilities> PropertyUtilities;
 };
diff --git a/Engine/Plugins/VirtualProduction/DMX/DMXProtocol/Source/DMXProtocolEditor/Private/DetailsCustomizations/DMXOutputPortConfigCustomization.cpp b/Engine/Plugins/VirtualProduction/DMX/DMXProtocol/Source/DMXProtocolEditor/Private/DetailsCustomizations/DMXOutputPortConfigCustomization.cpp
--- a/Engine/Plugins/VirtualProduction/DMX/DMXProtocol/Source/DMXProtocolEditor/Private/DetailsCustomizations/DMXOutputPortConfigCustomization.cpp
+++ b/Engine/Plugins/VirtualProduction/DMX/DMXProtocol/Source/DMXProtocolEditor/Private/DetailsCustomizations/DMXOutputPortConfigCustomization.cpp
@@ -2,27 +2,27 @@
 
 #include "DMXOutputPortConfigCustomization.h"
 
+#include "DetailWidgetRow.h"
 #include "DMXProtocolModule.h"
 #include "DMXProtocolSettings.h"
 #include "DMXProtocolTypes.h"
 #include "DMXProtocolUtils.h"
+#include "IDetailChildrenBuilder.h"
+#include "IDetailPropertyRow.h" 
+#include "IPropertyUtilities.h"
 #include "Interfaces/IDMXProtocol.h"
 #include "IO/DMXOutputPort.h"
 #include "IO/DMXOutputPortConfig.h"
 #include "IO/DMXPortManager.h"
+#include "Misc/Guid.h" 
+#include "ScopedTransaction.h"
+#include "Styling/AppStyle.h"
+#include "Widgets/Input/SCheckBox.h"
 #include "Widgets/SDMXCommunicationTypeComboBox.h"
 #include "Widgets/SDMXDelayEditWidget.h"
 #include "Widgets/SDMXIPAddressEditWidget.h"
 #include "Widgets/SDMXProtocolNameComboBox.h"
-
-#include "DetailWidgetRow.h"
-#include "IDetailChildrenBuilder.h"
-#include "IDetailPropertyRow.h" 
-#include "IPropertyUtilities.h"
-#include "ScopedTransaction.h"
-#include "Misc/Guid.h" 
 #include "Widgets/Text/STextBlock.h" 
-#include "Styling/AppStyle.h"
 
 
 #define LOCTEXT_NAMESPACE "DMXOutputPortConfigCustomization"
@@ -65,10 +65,15 @@
 	DeviceAddressHandle = PropertyHandles.FindChecked(FDMXOutputPortConfig::GetDeviceAddressPropertyNameChecked());
 	DelayHandle = PropertyHandles.FindChecked(FDMXOutputPortConfig::GetDelayPropertyNameChecked());
 	DelayFrameRateHandle = PropertyHandles.FindChecked(FDMXOutputPortConfig::GetDelayFrameRatePropertyNameChecked());
+	LocalUniverseStartHandle = PropertyHandles.FindChecked(FDMXInputPortConfig::GetLocalUniverseStartPropertyNameChecked());
+	ExternUniverseStartHandle = PropertyHandles.FindChecked(FDMXInputPortConfig::GetExternUniverseStartPropertyNameChecked());
+	IsExternUniverseStartEditableHandle = PropertyHandles.FindChecked(FDMXInputPortConfig::GetIsExternUnivereStartEditablePropertyNameChecked());
 	PortGuidHandle = PropertyHandles.FindChecked(FDMXOutputPortConfig::GetPortGuidPropertyNameChecked());
 
 	// Hande property changes
 	AutoCompleteDeviceAddressHandle->SetOnPropertyValueChanged(FSimpleDelegate::CreateSP(this, &FDMXOutputPortConfigCustomization::UpdateAutoCompleteDeviceAddressTextBox));
+	LocalUniverseStartHandle->SetOnPropertyValueChanged(FSimpleDelegate::CreateSP(this, &FDMXOutputPortConfigCustomization::OnLocalUniverseStartChanged));
+	IsExternUniverseStartEditableHandle->SetOnPropertyValueChanged(FSimpleDelegate::CreateSP(this, &FDMXOutputPortConfigCustomization::OnIsExternUniverseStartEditableChanged));
 
 	// Ports always need a valid Guid (cannot be blueprinted)
 	if (!GetPortGuid().IsValid())
@@ -94,7 +99,9 @@
 		}
 
 		// Don't add the bAutoCompeteDeviceAddressEnabled property
-		if (Iter.Value() == AutoCompleteDeviceAddressEnabledHandle)
+		// Don't add the bIsExternUniverseStartEditable property
+		if (Iter.Value() == AutoCompleteDeviceAddressEnabledHandle ||
+			Iter.Value() == IsExternUniverseStartEditableHandle)
 		{
 			continue;
 		}
@@ -118,6 +125,10 @@
 		{
 			GenerateDeviceAddressRow(PropertyRow);
 		}
+		else if (Iter.Value() == ExternUniverseStartHandle)
+		{
+			GenerateExternUniverseStartRow(PropertyRow);
+		}
 		else if (Iter.Key() == FDMXOutputPortConfig::GetDestinationAddressesPropertyNameChecked())
 		{
 			UpdateDestinationAddressesVisibility();
@@ -209,55 +220,6 @@
 		];
 }
 
-void FDMXOutputPortConfigCustomization::GenerateAutoCompleteDeviceAddressRow(IDetailPropertyRow& PropertyRow)
-{
-	// Mimic a Inline Edit Condition Toggle - Doing it directly via Meta Specifiers is not possible. -
-	// The Inline Edit Condition Toggle property wouldn't get properly serialized when being a Config property.
-	const TSharedRef<SWidget> AutoCompleteDeviceAddressEnableldPropertyValueWidget = AutoCompleteDeviceAddressEnabledHandle->CreatePropertyValueWidget();
-
-	const TSharedRef<SWidget> AutoCompleteDeviceAddressPropertyNameWidget = AutoCompleteDeviceAddressHandle->CreatePropertyNameWidget();
-	const TSharedRef<SWidget> AutoCompleteDeviceAddressPropertyValueWidget = AutoCompleteDeviceAddressHandle->CreatePropertyValueWidget();
-
-	AutoCompleteDeviceAddressPropertyNameWidget->SetEnabled(TAttribute<bool>::CreateLambda([this]()
-		{
-			return IsAutoCompleteDeviceAddressEnabled();
-		}));
-	AutoCompleteDeviceAddressPropertyValueWidget->SetEnabled(TAttribute<bool>::CreateLambda([this]()
-		{
-			return IsAutoCompleteDeviceAddressEnabled();
-		}));
-
-	PropertyRow.CustomWidget()
-		.NameContent()
-		[
-			SNew(SHorizontalBox)
-
-		// Edit Inline Condition Toggle Checkbox
-		+ SHorizontalBox::Slot()
-		.HAlign(HAlign_Left)
-		.VAlign(VAlign_Fill)
-		.AutoWidth()
-		[
-			AutoCompleteDeviceAddressEnableldPropertyValueWidget
-		]
-
-	// Property Name
-	+ SHorizontalBox::Slot()
-		.Padding(8.f, 0.f, 0.f, 0.f)
-		.HAlign(HAlign_Left)
-		.VAlign(VAlign_Fill)
-		.AutoWidth()
-		[
-			AutoCompleteDeviceAddressPropertyNameWidget
-		]
-		]
-	.ValueContent()
-		[
-			// Property Value
-			AutoCompleteDeviceAddressPropertyValueWidget
-		];
-}
-
 void FDMXOutputPortConfigCustomization::GenerateDeviceAddressRow(IDetailPropertyRow& PropertyRow)
 {
 	// Customizate the IPAddress property to show a combo box with available IP addresses
@@ -321,35 +283,52 @@
 		];
 }	
 
-void FDMXOutputPortConfigCustomization::GenerateDelayRow(IDetailPropertyRow& PropertyRow)
+void FDMXOutputPortConfigCustomization::GenerateAutoCompleteDeviceAddressRow(IDetailPropertyRow& PropertyRow)
 {
-	// Customizate the IPAddress property to show a combo box with available IP addresses
-	const double InitialDelay = GetDelay();
-	const TArray<FFrameRate> InitialDelayFrameRates = GetDelayFrameRates();
-	
-	FFrameRate InitialDelayFrameRate;
-	if (ensureMsgf(InitialDelayFrameRates.Num() == 1, TEXT("Multi editing Output Port Configs is not supported.")))
-	{
-		InitialDelayFrameRate = InitialDelayFrameRates[0];
-	}
+	// Mimic a Inline Edit Condition Toggle - Doing it directly via Meta Specifiers is not possible. -
+	// The Inline Edit Condition Toggle property wouldn't get properly serialized when being a Config property.
+	const TSharedRef<SWidget> AutoCompleteDeviceAddressEnableldPropertyValueWidget = AutoCompleteDeviceAddressEnabledHandle->CreatePropertyValueWidget();
 
-	TSharedPtr<SWidget> NameWidget;
-	TSharedPtr<SWidget> ValueWidget;
-	FDetailWidgetRow Row;
-	PropertyRow.GetDefaultWidgets(NameWidget, ValueWidget, Row);
+	const TSharedRef<SWidget> AutoCompleteDeviceAddressPropertyNameWidget = AutoCompleteDeviceAddressHandle->CreatePropertyNameWidget();
+	const TSharedRef<SWidget> AutoCompleteDeviceAddressPropertyValueWidget = AutoCompleteDeviceAddressHandle->CreatePropertyValueWidget();
 
+	AutoCompleteDeviceAddressPropertyNameWidget->SetEnabled(TAttribute<bool>::CreateLambda([this]()
+		{
+			return IsAutoCompleteDeviceAddressEnabled();
+		}));
+	AutoCompleteDeviceAddressPropertyValueWidget->SetEnabled(TAttribute<bool>::CreateLambda([this]()
+		{
+			return IsAutoCompleteDeviceAddressEnabled();
+		}));
+
 	PropertyRow.CustomWidget()
 		.NameContent()
 		[
-			NameWidget.ToSharedRef()
+			SNew(SHorizontalBox)
+
+		// Edit Inline Condition Toggle Checkbox
+		+ SHorizontalBox::Slot()
+		.HAlign(HAlign_Left)
+		.VAlign(VAlign_Fill)
+		.AutoWidth()
+		[
+			AutoCompleteDeviceAddressEnableldPropertyValueWidget
 		]
-		.ValueContent()
+
+	// Property Name
+	+ SHorizontalBox::Slot()
+		.Padding(8.f, 0.f, 0.f, 0.f)
+		.HAlign(HAlign_Left)
+		.VAlign(VAlign_Fill)
+		.AutoWidth()
 		[
-			SAssignNew(DelayEditWidget, SDMXDelayEditWidget)
-			.InitialDelay(InitialDelay)
-			.InitialDelayFrameRate(InitialDelayFrameRate)
-			.OnDelayChanged(this, &FDMXOutputPortConfigCustomization::OnDelayChanged)
-			.OnDelayFrameRateChanged(this, &FDMXOutputPortConfigCustomization::OnDelayFrameRateChanged)
+			AutoCompleteDeviceAddressPropertyNameWidget
+		]
+		]
+	.ValueContent()
+		[
+			// Property Value
+			AutoCompleteDeviceAddressPropertyValueWidget
 		];
 }
 
@@ -379,6 +358,144 @@
 	}
 }
 
+void FDMXOutputPortConfigCustomization::GenerateExternUniverseStartRow(IDetailPropertyRow& PropertyRow)
+{
+	const TAttribute<bool> CanEditExternUniverseAttribute = TAttribute<bool>::CreateLambda([this]()
+		{
+			bool bIsExternUniverseStartEditable;
+			if (IsExternUniverseStartEditableHandle->GetValue(bIsExternUniverseStartEditable) == FPropertyAccess::Success)
+			{
+				return bIsExternUniverseStartEditable ? true : false;
+			}
+			return false;
+		});
+
+	TSharedPtr<SWidget> ExternUniverseStartNameWidget;
+	TSharedPtr<SWidget> ExternUniverseStartValueWidget;
+	FDetailWidgetRow ExternUniverseStartRow;
+	PropertyRow.GetDefaultWidgets(ExternUniverseStartNameWidget, ExternUniverseStartValueWidget, ExternUniverseStartRow);
+
+	const FText IsExternUniverseEnabledTooltip = LOCTEXT("IsExternUniverseEnabledWarningTooltip", "Remaps the range of local Universes (Num Universes from Local Universe Start) to a different set of Universes at the Protocol level.\nFor example, for Art-Net this allows to map local Universe 1 to extern Universe 0.");
+
+	PropertyRow.CustomWidget()
+		.NameContent()
+		[
+			SNew(SHorizontalBox)
+
+			+ SHorizontalBox::Slot()
+			.AutoWidth()
+			.Padding(0.f, 0.f, 4.f, 0.f)
+			[
+				SNew(SCheckBox)
+				.ToolTipText(IsExternUniverseEnabledTooltip)
+				.IsChecked_Lambda([CanEditExternUniverseAttribute]()
+					{
+						return CanEditExternUniverseAttribute.Get() ? ECheckBoxState::Checked : ECheckBoxState::Unchecked;
+					})
+				.OnCheckStateChanged_Lambda([this](ECheckBoxState NewState)
+					{
+						const bool bIsExternUniverseStartEditable = NewState == ECheckBoxState::Checked;
+						IsExternUniverseStartEditableHandle->SetValue(bIsExternUniverseStartEditable);
+					})
+			]
+
+			+ SHorizontalBox::Slot()
+			.AutoWidth()
+			[
+				SNew(SBox)
+				.IsEnabled(CanEditExternUniverseAttribute)
+				[
+					ExternUniverseStartNameWidget.ToSharedRef()
+				]
+			]
+		]
+		.ValueContent()
+		[
+			SNew(SHorizontalBox)
+			.IsEnabled(CanEditExternUniverseAttribute)
+
+			+ SHorizontalBox::Slot()
+			.Padding(0.f, 0.f, 4.f, 0.f)
+			[
+				ExternUniverseStartValueWidget.ToSharedRef()
+			]
+
+			+ SHorizontalBox::Slot()
+			.AutoWidth()
+			.Padding(4.f)
+			[
+				SNew(SImage)
+				.Visibility_Lambda([CanEditExternUniverseAttribute]()
+					{
+						return CanEditExternUniverseAttribute.Get() ? EVisibility::Visible : EVisibility::Collapsed;
+					})
+				.Image(FCoreStyle::Get().GetBrush("Icons.Warning"))
+				.ToolTipText(IsExternUniverseEnabledTooltip)
+			]
+		];
+}
+
+void FDMXOutputPortConfigCustomization::OnLocalUniverseStartChanged()
+{
+	bool bCanEditExternUniverse;
+	int32 LocalUniverse;
+	if (IsExternUniverseStartEditableHandle->GetValue(bCanEditExternUniverse) == FPropertyAccess::Success &&
+		!bCanEditExternUniverse &&
+		LocalUniverseStartHandle->GetValue(LocalUniverse) == FPropertyAccess::Success)
+	{
+		ExternUniverseStartHandle->NotifyPreChange();
+		// Change the extern universe interactive only. Rely on the local universe start change completing the transaction.
+		ExternUniverseStartHandle->SetValue(LocalUniverse, EPropertyValueSetFlags::InteractiveChange);
+		ExternUniverseStartHandle->NotifyPostChange(EPropertyChangeType::Interactive);
+	}
+}
+
+void FDMXOutputPortConfigCustomization::OnIsExternUniverseStartEditableChanged()
+{
+	bool bCanEditExternUniverse;
+	int32 LocalUniverse;
+	if (IsExternUniverseStartEditableHandle->GetValue(bCanEditExternUniverse) == FPropertyAccess::Success &&
+		!bCanEditExternUniverse &&
+		LocalUniverseStartHandle->GetValue(LocalUniverse) == FPropertyAccess::Success)
+	{
+		ExternUniverseStartHandle->NotifyPreChange();
+		ExternUniverseStartHandle->SetValue(LocalUniverse);
+		ExternUniverseStartHandle->NotifyPostChange(EPropertyChangeType::ValueSet);
+	}
+}
+
+void FDMXOutputPortConfigCustomization::GenerateDelayRow(IDetailPropertyRow& PropertyRow)
+{
+	// Customizate the IPAddress property to show a combo box with available IP addresses
+	const double InitialDelay = GetDelay();
+	const TArray<FFrameRate> InitialDelayFrameRates = GetDelayFrameRates();
+	
+	FFrameRate InitialDelayFrameRate;
+	if (ensureMsgf(InitialDelayFrameRates.Num() == 1, TEXT("Multi editing Output Port Configs is not supported.")))
+	{
+		InitialDelayFrameRate = InitialDelayFrameRates[0];
+	}
+
+	TSharedPtr<SWidget> NameWidget;
+	TSharedPtr<SWidget> ValueWidget;
+	FDetailWidgetRow Row;
+	PropertyRow.GetDefaultWidgets(NameWidget, ValueWidget, Row);
+
+	PropertyRow.CustomWidget()
+		.NameContent()
+		[
+			NameWidget.ToSharedRef()
+		]
+		.ValueContent()
+		[
+			SAssignNew(DelayEditWidget, SDMXDelayEditWidget)
+			.InitialDelay(InitialDelay)
+			.InitialDelayFrameRate(InitialDelayFrameRate)
+			.OnDelayChanged(this, &FDMXOutputPortConfigCustomization::OnDelayChanged)
+			.OnDelayFrameRateChanged(this, &FDMXOutputPortConfigCustomization::OnDelayFrameRateChanged)
+		];
+}
+
 void FDMXOutputPortConfigCustomization::OnProtocolNameSelected()
 {
 	FName ProtocolName = ProtocolNameComboBox->GetSelectedProtocolName();
diff --git a/Engine/Plugins/VirtualProduction/DMX/DMXProtocol/Source/DMXProtocolEditor/Private/DetailsCustomizations/DMXOutputPortConfigCustomization.h b/Engine/Plugins/VirtualProduction/DMX/DMXProtocol/Source/DMXProtocolEditor/Private/DetailsCustomizations/DMXOutputPortConfigCustomization.h
--- a/Engine/Plugins/VirtualProduction/DMX/DMXProtocol/Source/DMXProtocolEditor/Private/DetailsCustomizations/DMXOutputPortConfigCustomization.h
+++ b/Engine/Plugins/VirtualProduction/DMX/DMXProtocol/Source/DMXProtocolEditor/Private/DetailsCustomizations/DMXOutputPortConfigCustomization.h
@@ -41,18 +41,27 @@ private:
 	/** Generates the customized Communication Type row */
 	void GenerateCommunicationTypeRow(IDetailPropertyRow& PropertyRow);
 
-	/** Generates the customized Auto Complete Device Address row */
-	void GenerateAutoCompleteDeviceAddressRow(IDetailPropertyRow& PropertyRow);
-
 	/** Generates the customized IP Address row */
 	void GenerateDeviceAddressRow(IDetailPropertyRow& PropertyRow);
 
-	/** Generates the customized Delay Seconds row */
-	void GenerateDelayRow(IDetailPropertyRow& PropertyRow);
+	/** Generates the customized Auto Complete Device Address row */
+	void GenerateAutoCompleteDeviceAddressRow(IDetailPropertyRow& PropertyRow);
 
 	/** Updates the Auto Complete Device Address Text Box */
 	void UpdateAutoCompleteDeviceAddressTextBox();
 
+	/** Generates the customized Extern Universe Start row */
+	void GenerateExternUniverseStartRow(IDetailPropertyRow& PropertyRow);
+
+	/** Called when the LocalUniverseStart property changed */
+	void OnLocalUniverseStartChanged();
+
+	/** Called when the bIsExternUniverseStartEditable property changed */
+	void OnIsExternUniverseStartEditableChanged();
+
+	/** Generates the customized Delay Seconds row */
+	void GenerateDelayRow(IDetailPropertyRow& PropertyRow);
+
 	/** Called when a Protocol Name was selected */
 	void OnProtocolNameSelected();
 
@@ -110,6 +119,15 @@ private:
 	/** Property handle to the CommunicationType property */
 	TSharedPtr<IPropertyHandle> CommunicationTypeHandle;
 
+	/** Property handle to the LocalUniverseStart property */
+	TSharedPtr<IPropertyHandle> LocalUniverseStartHandle;
+
+	/** Property handle to the ExternUniverseStart property */
+	TSharedPtr<IPropertyHandle> ExternUniverseStartHandle;
+
+	/** Property handle to the IsExternUniverseStartEditable property */
+	TSharedPtr<IPropertyHandle> IsExternUniverseStartEditableHandle;
+
 	/** Property handle to the Delay property */
 	TSharedPtr<IPropertyHandle> DelayHandle;
 
