DMX - Rename 'Protocol Name' to 'Protocol Type' in DMX project settings

diff --git a/Engine/Plugins/VirtualProduction/DMX/DMXProtocol/Source/DMXProtocol/Public/IO/DMXInputPortConfig.h b/Engine/Plugins/VirtualProduction/DMX/DMXProtocol/Source/DMXProtocol/Public/IO/DMXInputPortConfig.h
--- a/Engine/Plugins/VirtualProduction/DMX/DMXProtocol/Source/DMXProtocol/Public/IO/DMXInputPortConfig.h
+++ b/Engine/Plugins/VirtualProduction/DMX/DMXProtocol/Source/DMXProtocol/Public/IO/DMXInputPortConfig.h
@@ -117,7 +117,7 @@ protected:
 	FString PortName;
 
 	/** DMX Protocol */
-	UPROPERTY(Config, BlueprintReadOnly, EditDefaultsOnly, Category = "Port Config")
+	UPROPERTY(Config, BlueprintReadOnly, EditDefaultsOnly, Category = "Port Config", Meta = (DisplayName = "Protocol Type"))
 	FName ProtocolName = FDMXProtocolModule::DefaultProtocolArtNetName;
 
 	/** The type of communication used with this port */
diff --git a/Engine/Plugins/VirtualProduction/DMX/DMXProtocol/Source/DMXProtocol/Public/IO/DMXOutputPortConfig.h b/Engine/Plugins/VirtualProduction/DMX/DMXProtocol/Source/DMXProtocol/Public/IO/DMXOutputPortConfig.h
--- a/Engine/Plugins/VirtualProduction/DMX/DMXProtocol/Source/DMXProtocol/Public/IO/DMXOutputPortConfig.h
+++ b/Engine/Plugins/VirtualProduction/DMX/DMXProtocol/Source/DMXProtocol/Public/IO/DMXOutputPortConfig.h
@@ -134,7 +134,7 @@ protected:
 	FString PortName;
 
 	/** DMX Protocol */
-	UPROPERTY(Config, BlueprintReadOnly, EditDefaultsOnly, Category = "Port Config")
+	UPROPERTY(Config, BlueprintReadOnly, EditDefaultsOnly, Category = "Port Config", Meta = (DisplayName = "Protocol Type"))
 	FName ProtocolName = FDMXProtocolModule::DefaultProtocolArtNetName;
 
 	/** The type of communication used with this port */
