DMX - Fix possible race condition when stopping output port threads

diff --git a/Engine/Plugins/VirtualProduction/DMX/DMXProtocol/Source/DMXProtocol/Private/IO/DMXOutputPort.cpp b/Engine/Plugins/VirtualProduction/DMX/DMXProtocol/Source/DMXProtocol/Private/IO/DMXOutputPort.cpp
--- a/Engine/Plugins/VirtualProduction/DMX/DMXProtocol/Source/DMXProtocol/Private/IO/DMXOutputPort.cpp
+++ b/Engine/Plugins/VirtualProduction/DMX/DMXProtocol/Source/DMXProtocol/Private/IO/DMXOutputPort.cpp
@@ -536,6 +536,9 @@ FDMXOutputPortSharedRef FDMXOutputPort::CreateFromConfig(FDMXOutputPortConfig& O
 
 	FDMXOutputPortSharedRef NewOutputPort = MakeShared<FDMXOutputPort, ESPMode::ThreadSafe>();
 
+	NewOutputPort->SendDMXEvent = FPlatformProcess::GetSynchEventFromPool();
+	check(NewOutputPort->SendDMXEvent != nullptr);
+
 	NewOutputPort->PortGuid = OutputPortConfig.GetPortGuid();
 
 	UDMXProtocolSettings* Settings = GetMutableDefault<UDMXProtocolSettings>();
@@ -569,6 +572,9 @@ FDMXOutputPort::~FDMXOutputPort()
 		delete Thread;
 	}
 
+	FPlatformProcess::ReturnSynchEventToPool(SendDMXEvent);
+	SendDMXEvent = nullptr;
+
 	UE_LOG(LogDMXProtocol, VeryVerbose, TEXT("Destroyed output port %s"), *PortName);
 }
 
@@ -942,9 +948,6 @@ bool FDMXOutputPort::Init()
 
 uint32 FDMXOutputPort::Run()
 {
-	SendDMXEvent = FPlatformProcess::GetSynchEventFromPool();
-	check(SendDMXEvent != nullptr);
-
 	while (!bStopping)
 	{
 		const double StartTime = FPlatformTime::Seconds();
@@ -956,18 +959,17 @@ uint32 FDMXOutputPort::Run()
 
 		ProcessSendDMX();
 
+		const double Interval = 1.0 / SendRate;
 		const double EndTime = FPlatformTime::Seconds();
-		const double WaitTimeMs = ((1.0 / SendRate) - (EndTime - StartTime)) * 1000.0;
-		
-		if (WaitTimeMs > 0.0 && WaitTimeMs < std::numeric_limits<uint32>::max())
+		const double Elapsed = EndTime - StartTime;
+
+		if (Interval > Elapsed)
 		{
-			SendDMXEvent->Wait(WaitTimeMs);
+			const FTimespan WaitTimespan = FTimespan::FromSeconds(Interval - Elapsed);
+			SendDMXEvent->Wait(WaitTimespan);
 		}
 	}
 
-	FPlatformProcess::ReturnSynchEventToPool(SendDMXEvent);
-	SendDMXEvent = nullptr;
-
 	return 0;
 }
 
diff --git a/Engine/Plugins/VirtualProduction/DMX/DMXProtocol/Source/DMXProtocol/Public/IO/DMXOutputPort.h b/Engine/Plugins/VirtualProduction/DMX/DMXProtocol/Source/DMXProtocol/Public/IO/DMXOutputPort.h
--- a/Engine/Plugins/VirtualProduction/DMX/DMXProtocol/Source/DMXProtocol/Public/IO/DMXOutputPort.h
+++ b/Engine/Plugins/VirtualProduction/DMX/DMXProtocol/Source/DMXProtocol/Public/IO/DMXOutputPort.h
@@ -88,15 +88,11 @@ struct DMXPROTOCOL_API FDMXSignalFragment
 
 
 /**
- * Higher level abstraction of a DMX input hiding networking specific and protocol specific complexity.
+ * Higher level abstraction of a DMX input hiding networking specific and protocol specific complexity. Ment to be constructed by DMXPortManger only.
  *
- * Use SendDMXFragment method to send DMX.
+ * Use SendDMX method to send DMX.
  *
- * To loopback outputs refer to DMXRawListener and DMXTickedUniverseListener.
- *
- * Can only be constructed via DMXPortManger, see FDMXPortManager::CreateOutputPort and FDMXPortManager::CreateOutputPortFromConfig.
- * 
- * Note, internally locks have to be acquired in order as they're declared in the header.
+ * To loopback outputs, refer to DMXRawListener.
  */
 class DMXPROTOCOL_API FDMXOutputPort
 	: public FDMXPort
