DMX - Fix non-standard header include

diff --git a/Engine/Plugins/VirtualProduction/DMX/DMXProtocol/Source/DMXProtocol/Private/IO/DMXOutputPort.cpp b/Engine/Plugins/VirtualProduction/DMX/DMXProtocol/Source/DMXProtocol/Private/IO/DMXOutputPort.cpp
--- a/Engine/Plugins/VirtualProduction/DMX/DMXProtocol/Source/DMXProtocol/Private/IO/DMXOutputPort.cpp
+++ b/Engine/Plugins/VirtualProduction/DMX/DMXProtocol/Source/DMXProtocol/Private/IO/DMXOutputPort.cpp
@@ -14,7 +14,7 @@
 #include "HAL/Event.h"
 #include "HAL/IConsoleManager.h"
 #include "HAL/RunnableThread.h"
-#include <limits.h>
+#include <limits>
 #include "Misc/FrameRate.h"
 
 
