DMX - Remove deprecated members from FDMXAttributeName, fix C4072 warning, remove CoreMinimal include

diff --git a/Engine/Plugins/VirtualProduction/DMX/DMXProtocol/Source/DMXProtocol/Private/DMXAttribute.cpp b/Engine/Plugins/VirtualProduction/DMX/DMXProtocol/Source/DMXProtocol/Private/DMXAttribute.cpp
--- a/Engine/Plugins/VirtualProduction/DMX/DMXProtocol/Source/DMXProtocol/Private/DMXAttribute.cpp
+++ b/Engine/Plugins/VirtualProduction/DMX/DMXProtocol/Source/DMXProtocol/Private/DMXAttribute.cpp
@@ -1,15 +1,11 @@
 // Copyright Epic Games, Inc. All Rights Reserved.
 
 #include "DMXAttribute.h"
+
 #include "DMXProtocolModule.h"
 #include "DMXProtocolSettings.h"
-
 #include "Modules/ModuleManager.h"
 
-PRAGMA_DISABLE_DEPRECATION_WARNINGS
-const bool FDMXAttributeName::bCanBeNone = true;
-FSimpleMulticastDelegate FDMXAttributeName::OnValuesChanged;
-PRAGMA_ENABLE_DEPRECATION_WARNINGS
 
 FDMXAttributeName::FDMXAttributeName()
 {
@@ -44,28 +40,6 @@ void FDMXAttributeName::SetFromName(const FName& InName)
 	*this = InName;
 }
 
-FDMXAttribute FDMXAttributeName::GetAttribute() const
-{
-	FDMXAttribute Attribute;
-	Attribute.Name = Name;
-	return Attribute;
-	const UDMXProtocolSettings* DMXSettings = GetDefault<UDMXProtocolSettings>();
-	if (!DMXSettings)
-	{
-		return Attribute;
-	}
-
-	for (const FDMXAttribute& SettingsAttribute : DMXSettings->Attributes)
-	{
-		if (SettingsAttribute.Name.IsEqual(Name))
-		{
-			return SettingsAttribute;
-		}
-	}
-
-	return Attribute;
-}
-
 TArray<FName> FDMXAttributeName::GetPredefinedValues()
 {
 	TArray<FName> Result;
@@ -82,11 +56,6 @@ TArray<FName> FDMXAttributeName::GetPredefinedValues()
 	return Result;
 }
 
-TArray<FName> FDMXAttributeName::GetPossibleValues()
-{
-	return GetPredefinedValues();
-}
-
 FString UDMXAttributeNameConversions::Conv_DMXAttributeToString(const FDMXAttributeName& InAttribute)
 {
 	return InAttribute.Name.ToString();
diff --git a/Engine/Plugins/VirtualProduction/DMX/DMXProtocol/Source/DMXProtocol/Public/DMXAttribute.h b/Engine/Plugins/VirtualProduction/DMX/DMXProtocol/Source/DMXProtocol/Public/DMXAttribute.h
--- a/Engine/Plugins/VirtualProduction/DMX/DMXProtocol/Source/DMXProtocol/Public/DMXAttribute.h
+++ b/Engine/Plugins/VirtualProduction/DMX/DMXProtocol/Source/DMXProtocol/Public/DMXAttribute.h
@@ -2,15 +2,14 @@
 
 #pragma once
 
-#include "CoreMinimal.h"
 #include "DMXNameListItem.h"
-#include "Misc/Crc.h"
-#include <UObject/NameTypes.h>
-
 #include "Kismet/BlueprintFunctionLibrary.h"
+#include "Misc/Crc.h"
+#include "UObject/NameTypes.h"
 
 #include "DMXAttribute.generated.h"
 
+
 USTRUCT()
 struct DMXPROTOCOL_API FDMXAttribute
 {
@@ -76,22 +75,6 @@ public:
 
 	/** Returns the predefined values */
 	static TArray<FName> GetPredefinedValues();
-
-	UE_DEPRECATED(5.1, "Deprecated since attributes now can take any name. If it is desired to resolve against default attributes from project settings, do a lookup on UDMXProtocolSettings::Attributes")
-	FDMXAttribute GetAttribute() const;
-
-	//////////////////////////////////////////////////////////////////
-	// Deprecated members originating from deprecated FDMXNameListItem
-	UE_DEPRECATED(5.1, "Please use GetPredefinedValues() instead.")
-	static TArray<FName> GetPossibleValues();
-
-	// Deprecated 5.1 without replacement. Always true
-	UE_DEPRECATED(5.1, "Obsolete. Value 'None' is always allowed.")
-	static const bool bCanBeNone;
-
-	// Deprecated 5.1 without replacement. To listen to changes of default defined protocols, please refer to UDMXProtocolSettings::GetOnDefaultAttributesChanged
-	UE_DEPRECATED(5.1, "To listen to changes of default defined protocols, please refer to UDMXProtocolSettings::GetOnDefaultAttributesChanged().")
-	static FSimpleMulticastDelegate OnValuesChanged;
 };
 
 inline uint32 GetTypeHash(const FDMXAttributeName& DMXNameListItem)
