DMX: Fix issues with MVR import

* Fix importing more than one MVR file can lead to a crash when the update ports option is checked
* Fix GDTFs imported via MVR don't show in the content browser

diff --git a/Engine/Plugins/VirtualProduction/DMX/DMXEngine/Source/DMXEditor/Private/Factories/DMXLibraryFromMVRImportOptions.cpp b/Engine/Plugins/VirtualProduction/DMX/DMXEngine/Source/DMXEditor/Private/Factories/DMXLibraryFromMVRImportOptions.cpp
--- a/Engine/Plugins/VirtualProduction/DMX/DMXEngine/Source/DMXEditor/Private/Factories/DMXLibraryFromMVRImportOptions.cpp
+++ b/Engine/Plugins/VirtualProduction/DMX/DMXEngine/Source/DMXEditor/Private/Factories/DMXLibraryFromMVRImportOptions.cpp
@@ -77,9 +77,7 @@ void UDMXLibraryFromMVRImportOptions::ApplyOptions(UDMXLibrary* DMXLibrary)
 			ProtocolSettings->InputPortConfigs.Add(NewConfig);
 			ProtocolSettings->PostEditChange();
 		}
-		ProtocolSettings->SaveConfig();
-
-		FDMXPortManager::Get().UpdateFromProtocolSettings();
+		ProtocolSettings->TryUpdateDefaultConfigFile();
 	}
 
 	if (bUpdateOutputPort)
@@ -114,8 +112,8 @@ void UDMXLibraryFromMVRImportOptions::ApplyOptions(UDMXLibrary* DMXLibrary)
 			ProtocolSettings->OutputPortConfigs.Add(NewConfig);
 			ProtocolSettings->PostEditChange();
 		}
-		ProtocolSettings->SaveConfig();
-
-		FDMXPortManager::Get().UpdateFromProtocolSettings();
+		ProtocolSettings->TryUpdateDefaultConfigFile();
 	}
+
+	FDMXPortManager::Get().UpdateFromProtocolSettings();
 }
diff --git a/Engine/Plugins/VirtualProduction/DMX/DMXEngine/Source/DMXEditor/Private/Factories/DMXLibraryFromMVRImporter.cpp b/Engine/Plugins/VirtualProduction/DMX/DMXEngine/Source/DMXEditor/Private/Factories/DMXLibraryFromMVRImporter.cpp
--- a/Engine/Plugins/VirtualProduction/DMX/DMXEngine/Source/DMXEditor/Private/Factories/DMXLibraryFromMVRImporter.cpp
+++ b/Engine/Plugins/VirtualProduction/DMX/DMXEngine/Source/DMXEditor/Private/Factories/DMXLibraryFromMVRImporter.cpp
@@ -454,7 +454,7 @@ TArray<UDMXImportGDTF*> FDMXLibraryFromMVRImporter::CreateGDTFs(UObject* InParen
 				Package->FullyLoad();
 
 				bool bCancelled;
-				UObject* NewGDTFObject = GDTFFactory->FactoryCreateFile(UDMXImportGDTF::StaticClass(), Package, FName(*FixtureNode->GDTFSpec), InFlags | RF_Public, ScopedUnzipGDTF.TempFilePathAndName, nullptr, GWarn, bCancelled);
+				UObject* NewGDTFObject = GDTFFactory->FactoryCreateFile(UDMXImportGDTF::StaticClass(), Package, *GDTFName, InFlags | RF_Public, ScopedUnzipGDTF.TempFilePathAndName, nullptr, GWarn, bCancelled);
 
 				if (UDMXImportGDTF* NewGDTF = Cast<UDMXImportGDTF>(NewGDTFObject))
 				{
diff --git a/Engine/Plugins/VirtualProduction/DMX/DMXProtocol/Source/DMXProtocolEditor/Private/DetailsCustomizations/DMXPortReferenceCustomizationBase.cpp b/Engine/Plugins/VirtualProduction/DMX/DMXProtocol/Source/DMXProtocolEditor/Private/DetailsCustomizations/DMXPortReferenceCustomizationBase.cpp
--- a/Engine/Plugins/VirtualProduction/DMX/DMXProtocol/Source/DMXProtocolEditor/Private/DetailsCustomizations/DMXPortReferenceCustomizationBase.cpp
+++ b/Engine/Plugins/VirtualProduction/DMX/DMXProtocol/Source/DMXProtocolEditor/Private/DetailsCustomizations/DMXPortReferenceCustomizationBase.cpp
@@ -110,8 +110,6 @@ void FDMXPortReferenceCustomizationBase::CustomizeChildren(TSharedRef<IPropertyH
 
 void FDMXPortReferenceCustomizationBase::OnPortSelected()
 {
-	check(PortSelector.IsValid());
-	
 	ApplySelectedPortGuid();
 }
 
@@ -178,7 +176,11 @@ void FDMXPortReferenceCustomizationBase::ApplySelectedPortGuid()
 	if (!GetPortGuids().Contains(SelectedGuid))
 	{
 		const TSharedPtr<IPropertyHandle>& PortGuidHandle = GetPortGuidHandle();
-		check(PortGuidHandle.IsValid() && PortGuidHandle->IsValidHandle());
+		if (!PortGuidHandle.IsValid() || !PortGuidHandle->IsValidHandle())
+		{
+			// May be invalid while port arrays are changing
+			return;
+		}
 
 		PortGuidHandle->NotifyPreChange();
 
@@ -197,7 +199,11 @@ void FDMXPortReferenceCustomizationBase::ApplySelectedPortGuid()
 TArray<FGuid> FDMXPortReferenceCustomizationBase::GetPortGuids() const
 {
 	const TSharedPtr<IPropertyHandle>& PortGuidHandle = GetPortGuidHandle();
-	check(PortGuidHandle.IsValid() && PortGuidHandle->IsValidHandle());
+	if (!PortGuidHandle.IsValid() || !PortGuidHandle->IsValidHandle())
+	{			
+		// May be invalid while port arrays are changing
+		return TArray<FGuid>();
+	}
 
 	TArray<void*> RawDataArray;
 	PortGuidHandle->AccessRawData(RawDataArray);
