DMX: Fix crash during shutdown when testing for valid DMX port manager. Instead accept the port manager may never have been started up, or is shut down more than once.

diff --git a/Engine/Plugins/VirtualProduction/DMX/DMXProtocol/Source/DMXProtocol/Private/IO/DMXPortManager.cpp b/Engine/Plugins/VirtualProduction/DMX/DMXProtocol/Source/DMXProtocol/Private/IO/DMXPortManager.cpp
--- a/Engine/Plugins/VirtualProduction/DMX/DMXProtocol/Source/DMXProtocol/Private/IO/DMXPortManager.cpp
+++ b/Engine/Plugins/VirtualProduction/DMX/DMXProtocol/Source/DMXProtocol/Private/IO/DMXPortManager.cpp
@@ -333,7 +333,5 @@ void FDMXPortManager::ShutdownManager()
 {
 	UE_LOG(LogDMXProtocol, Verbose, TEXT("Shutting down DMXPortManager"));
 
-	check(CurrentManager.IsValid());
-	
 	CurrentManager.Reset();
 }
