DMX: Default DMX Signal Priority to 100

#rnx

diff --git a/Engine/Plugins/VirtualProduction/DMX/DMXProtocol/Source/DMXProtocol/Public/DMXProtocolTypes.h b/Engine/Plugins/VirtualProduction/DMX/DMXProtocol/Source/DMXProtocol/Public/DMXProtocolTypes.h
--- a/Engine/Plugins/VirtualProduction/DMX/DMXProtocol/Source/DMXProtocol/Public/DMXProtocolTypes.h
+++ b/Engine/Plugins/VirtualProduction/DMX/DMXProtocol/Source/DMXProtocol/Public/DMXProtocolTypes.h
@@ -65,7 +65,7 @@ public:
 
 	int32 ExternUniverseID;
 
-	int32 Priority;
+	int32 Priority = 100;
 
 	TArray<uint8> ChannelData;
 };
