DMX: Initially expand port arrays in project settings

#rnx

diff --git a/Engine/Plugins/VirtualProduction/DMX/DMXProtocol/Source/DMXProtocolEditor/Private/DetailsCustomizations/DMXInputPortConfigCustomization.cpp b/Engine/Plugins/VirtualProduction/DMX/DMXProtocol/Source/DMXProtocolEditor/Private/DetailsCustomizations/DMXInputPortConfigCustomization.cpp
--- a/Engine/Plugins/VirtualProduction/DMX/DMXProtocol/Source/DMXProtocolEditor/Private/DetailsCustomizations/DMXInputPortConfigCustomization.cpp
+++ b/Engine/Plugins/VirtualProduction/DMX/DMXProtocol/Source/DMXProtocolEditor/Private/DetailsCustomizations/DMXInputPortConfigCustomization.cpp
@@ -45,6 +45,11 @@ void FDMXInputPortConfigCustomization::CustomizeHeader(TSharedRef<IPropertyHandl
 
 void FDMXInputPortConfigCustomization::CustomizeChildren(TSharedRef<IPropertyHandle> StructPropertyHandle, class IDetailChildrenBuilder& ChildBuilder, IPropertyTypeCustomizationUtils& StructCustomizationUtils)
 {
+	if (!StructPropertyHandle->IsExpanded())
+	{
+		StructPropertyHandle->SetExpanded(true);
+	}
+
 	// Retrieve structure's child properties
 	uint32 NumChildren;
 	StructPropertyHandle->GetNumChildren(NumChildren);
diff --git a/Engine/Plugins/VirtualProduction/DMX/DMXProtocol/Source/DMXProtocolEditor/Private/DetailsCustomizations/DMXOutputPortConfigCustomization.cpp b/Engine/Plugins/VirtualProduction/DMX/DMXProtocol/Source/DMXProtocolEditor/Private/DetailsCustomizations/DMXOutputPortConfigCustomization.cpp
--- a/Engine/Plugins/VirtualProduction/DMX/DMXProtocol/Source/DMXProtocolEditor/Private/DetailsCustomizations/DMXOutputPortConfigCustomization.cpp
+++ b/Engine/Plugins/VirtualProduction/DMX/DMXProtocol/Source/DMXProtocolEditor/Private/DetailsCustomizations/DMXOutputPortConfigCustomization.cpp
@@ -45,6 +45,11 @@ void FDMXOutputPortConfigCustomization::CustomizeHeader(TSharedRef<IPropertyHand
 
 void FDMXOutputPortConfigCustomization::CustomizeChildren(TSharedRef<IPropertyHandle> StructPropertyHandle, class IDetailChildrenBuilder& ChildBuilder, IPropertyTypeCustomizationUtils& StructCustomizationUtils)
 {
+	if (!StructPropertyHandle->IsExpanded())
+	{	
+		StructPropertyHandle->SetExpanded(true);
+	}
+
 	// Retrieve structure's child properties
 	uint32 NumChildren;
 	StructPropertyHandle->GetNumChildren(NumChildren);
