DMX: Fix Attributes defaulted to the first entry of Attributes in project settings, leading to changed defaults if the first entry changed

diff --git a/Engine/Plugins/VirtualProduction/DMX/DMXProtocol/Source/DMXProtocol/Private/DMXAttribute.cpp b/Engine/Plugins/VirtualProduction/DMX/DMXProtocol/Source/DMXProtocol/Private/DMXAttribute.cpp
--- a/Engine/Plugins/VirtualProduction/DMX/DMXProtocol/Source/DMXProtocol/Private/DMXAttribute.cpp
+++ b/Engine/Plugins/VirtualProduction/DMX/DMXProtocol/Source/DMXProtocol/Private/DMXAttribute.cpp
@@ -3,26 +3,18 @@
 #include "DMXAttribute.h"
 
 #include "DMXProtocolModule.h"
+#include "DMXProtocolObjectVersion.h"
 #include "DMXProtocolSettings.h"
 #include "Modules/ModuleManager.h"
 
-
+// UE-224455. Before 5.5 the Name member was defaulted to the first entry of the Attributes set in project settings, 
+// which is variable implicitly. Since CDOs can define the struct default individually, it is important to keep the
+// original attribute name "Color". This is to ensure CDOs of objects created before 5.5 get the correct default value.
+// 
+// Also see FDMXAttributeName::Serialize where instances that use a structure serializer are handled.
 FDMXAttributeName::FDMXAttributeName()
+	: Name("Color")
 {
-	// This depends on the FDMXProtocolModule and can be called
-	// on CDO creation, when the module might not be available yet.
-	// So we first check if it is available.
-	const IModuleInterface* DMXProtocolModule = FModuleManager::Get().GetModule("DMXProtocol");
-	if (DMXProtocolModule != nullptr)
-	{
-		if (const UDMXProtocolSettings* DMXSettings = GetDefault<UDMXProtocolSettings>())
-		{
-			if (DMXSettings->Attributes.Num() > 0)
-			{
-				Name = DMXSettings->Attributes.begin()->Name;
-			}
-		}
-	}
 }
 
 FDMXAttributeName::FDMXAttributeName(const FDMXAttribute& InAttribute)
@@ -56,6 +48,34 @@ TArray<FName> FDMXAttributeName::GetPredefinedValues()
 	return Result;
 }
 
+bool FDMXAttributeName::Serialize(FArchive& Ar)
+{
+#if WITH_EDITORONLY_DATA
+	Ar.UsingCustomVersion(FDMXProtocolObjectVersion::GUID);
+	if (Ar.CustomVer(FDMXProtocolObjectVersion::GUID) < FDMXProtocolObjectVersion::FixAttributeNameDefaultValue)
+	{			
+		// UE-224455. The Name member was defaulted to the first entry of the Attributes set in project settings, 
+		// which is variable implicitly. For old projects keep this behaviour when upgrading to 5.5.
+
+		const IModuleInterface* DMXProtocolModule = FModuleManager::Get().GetModule("DMXProtocol");
+		if (DMXProtocolModule != nullptr)
+		{
+			if (const UDMXProtocolSettings* DMXSettings = GetDefault<UDMXProtocolSettings>())
+			{
+				if (DMXSettings->Attributes.Num() > 0)
+				{
+					Name = DMXSettings->Attributes.begin()->Name;
+				}
+			}
+		}
+	}
+#endif // WITH_EDITORONLY_DATA
+
+	// Serialize is only implemented to recall the right default for versions before 5.5.
+	// Return false to leave it to the outer serializer to perform the actual serialization.
+	return false;
+}
+
 FString UDMXAttributeNameConversions::Conv_DMXAttributeToString(const FDMXAttributeName& InAttribute)
 {
 	return InAttribute.Name.ToString();
diff --git a/Engine/Plugins/VirtualProduction/DMX/DMXProtocol/Source/DMXProtocol/Public/DMXAttribute.h b/Engine/Plugins/VirtualProduction/DMX/DMXProtocol/Source/DMXProtocol/Public/DMXAttribute.h
--- a/Engine/Plugins/VirtualProduction/DMX/DMXProtocol/Source/DMXProtocol/Public/DMXAttribute.h
+++ b/Engine/Plugins/VirtualProduction/DMX/DMXProtocol/Source/DMXProtocol/Public/DMXAttribute.h
@@ -5,10 +5,15 @@
 #include "DMXNameListItem.h"
 #include "Kismet/BlueprintFunctionLibrary.h"
 #include "Misc/Crc.h"
+#include "Misc/Optional.h"
 #include "UObject/NameTypes.h"
 
 #include "DMXAttribute.generated.h"
 
+namespace UE::DMX::Attribute
+{
+	class IDMXAttributeEditorContext;
+}
 
 USTRUCT()
 struct DMXPROTOCOL_API FDMXAttribute
@@ -50,11 +55,9 @@ struct DMXPROTOCOL_API FDMXAttributeName
 {
 	GENERATED_BODY()
 
+	using IDMXAttributeEditorContext = UE::DMX::Attribute::IDMXAttributeEditorContext;
+
 public:
-	/**
-	 * Default constructor. Assigns Attribute to the first available
-	 * Attribute from the plugin settings
-	 */
 	FDMXAttributeName();
 
 	/** Construct from an Attribute */
@@ -75,6 +78,19 @@ public:
 
 	/** Returns the predefined values */
 	static TArray<FName> GetPredefinedValues();
+
+	/** Custom serializer to maintain the upgrade path */
+	bool Serialize(FArchive& Ar);
+};
+
+template<>
+struct TStructOpsTypeTraits<FDMXAttributeName>
+	: public TStructOpsTypeTraitsBase2<FDMXAttributeName>
+{
+	enum
+	{
+		WithSerializer = true
+	};
 };
 
 inline uint32 GetTypeHash(const FDMXAttributeName& DMXNameListItem)
diff --git a/Engine/Plugins/VirtualProduction/DMX/DMXProtocol/Source/DMXProtocol/Public/DMXProtocolObjectVersion.h b/Engine/Plugins/VirtualProduction/DMX/DMXProtocol/Source/DMXProtocol/Public/DMXProtocolObjectVersion.h
--- a/Engine/Plugins/VirtualProduction/DMX/DMXProtocol/Source/DMXProtocol/Public/DMXProtocolObjectVersion.h
+++ b/Engine/Plugins/VirtualProduction/DMX/DMXProtocol/Source/DMXProtocol/Public/DMXProtocolObjectVersion.h
@@ -15,6 +15,9 @@ struct FDMXProtocolObjectVersion
 		// Update Ports to allow for an array of Unicast Addresses instead of just one
 		OutputPortSupportsManyUnicastAddresses,
 
+		// 5.5 fix UE-224455, variable default for FDXMAttribute::Name
+		FixAttributeNameDefaultValue,
+
 		// -----<new versions can be added above this line>-------------------------------------------------
 		VersionPlusOne,
 		LatestVersion = VersionPlusOne - 1
