DMX: Add a warning when an Output Port IP Address cannot be found on the local machine, suppress warnings for invalid IP addresses in Input and Output Ports when a cook commandlet is run

diff --git a/Engine/Plugins/VirtualProduction/DMX/DMXProtocol/Source/DMXProtocol/Private/IO/DMXInputPort.cpp b/Engine/Plugins/VirtualProduction/DMX/DMXProtocol/Source/DMXProtocol/Private/IO/DMXInputPort.cpp
--- a/Engine/Plugins/VirtualProduction/DMX/DMXProtocol/Source/DMXProtocol/Private/IO/DMXInputPort.cpp
+++ b/Engine/Plugins/VirtualProduction/DMX/DMXProtocol/Source/DMXProtocol/Private/IO/DMXInputPort.cpp
@@ -534,7 +534,7 @@ void FDMXInputPort::UpdateFromConfig(FDMXInputPortConfig& InOutInputPortConfig,
 		{
 			return *IPAddress == DeviceAddress;
 		});
-	if (!bValidDeviceAddress)
+	if (!bValidDeviceAddress && !IsRunningCookCommandlet())
 	{
 		UE_LOG(LogDMXProtocol, Warning, TEXT("Cannot register Input Port %s. Device Address '%s' does not exist on the local machine."), *PortName, *DeviceAddress);
 	}
diff --git a/Engine/Plugins/VirtualProduction/DMX/DMXProtocol/Source/DMXProtocol/Private/IO/DMXOutputPort.cpp b/Engine/Plugins/VirtualProduction/DMX/DMXProtocol/Source/DMXProtocol/Private/IO/DMXOutputPort.cpp
--- a/Engine/Plugins/VirtualProduction/DMX/DMXProtocol/Source/DMXProtocol/Private/IO/DMXOutputPort.cpp
+++ b/Engine/Plugins/VirtualProduction/DMX/DMXProtocol/Source/DMXProtocol/Private/IO/DMXOutputPort.cpp
@@ -683,6 +683,15 @@ void FDMXOutputPort::UpdateFromConfig(FDMXOutputPortConfig& InOutOutputPortConfi
 
 	CommunicationDeterminator.SetLoopbackToEngine(OutputPortConfig.NeedsLoopbackToEngine());
 
+	const bool bValidDeviceAddress = FDMXProtocolUtils::GetLocalNetworkInterfaceCardIPs().ContainsByPredicate([this](const TSharedPtr<FString>& IPAddress)
+		{
+			return *IPAddress == DeviceAddress;
+		});
+	if (!bValidDeviceAddress && !IsRunningCookCommandlet())
+	{
+		UE_LOG(LogDMXProtocol, Warning, TEXT("Cannot register Output Port %s. Device Address '%s' does not exist on the local machine."), *PortName, *DeviceAddress);
+	}
+
 	// Re-register the port if required
 	if (bNeedsUpdateRegistration)
 	{
