DMX: Improve implementation and robustness of the sACN protocol implementation

* Fix sACN cannot be used with multiple network cards
* Fix an issue where ports using sACN may no loger be functional when changing port settings
* Improve platform specific implementations
* Add a CVar to set sACN multicast Time-to-Live for output ports

diff --git a/Engine/Plugins/VirtualProduction/DMX/DMXProtocol/Source/DMXProtocol/Private/IO/DMXInputPort.cpp b/Engine/Plugins/VirtualProduction/DMX/DMXProtocol/Source/DMXProtocol/Private/IO/DMXInputPort.cpp
--- a/Engine/Plugins/VirtualProduction/DMX/DMXProtocol/Source/DMXProtocol/Private/IO/DMXInputPort.cpp
+++ b/Engine/Plugins/VirtualProduction/DMX/DMXProtocol/Source/DMXProtocol/Private/IO/DMXInputPort.cpp
@@ -2,6 +2,7 @@
 
 #include "IO/DMXInputPort.h"
 
+#include "Async/Async.h"
 #include "DMXProtocolSettings.h"
 #include "DMXProtocolUtils.h"
 #include "DMXStats.h"
@@ -530,22 +531,30 @@ void FDMXInputPort::UpdateFromConfig(FDMXInputPortConfig& InOutInputPortConfig,
 	PriorityStrategy = InputPortConfig.GetPortPriorityStrategy();
 	Priority = InputPortConfig.GetPriority();
 
-	const bool bValidDeviceAddress = FDMXProtocolUtils::GetLocalNetworkInterfaceCardIPs().ContainsByPredicate([this](const TSharedPtr<FString>& IPAddress)
-		{
-			return *IPAddress == DeviceAddress;
-		});
-	if (!bValidDeviceAddress && !IsRunningCookCommandlet())
+	// Only update on the next tick to give the OS time to fully release previous sockets
+	if (bPendingRegistration)
 	{
-		UE_LOG(LogDMXProtocol, Warning, TEXT("Cannot register Input Port %s. Device Address '%s' does not exist on the local machine."), *PortName, *DeviceAddress);
+		return;
 	}
+	bPendingRegistration = true;
 
-	// Re-register the port if required
-	if (bNeedsUpdateRegistration && bValidDeviceAddress)
-	{
-		Register();
-	}
+	AsyncTask(ENamedThreads::GameThread, [bNeedsUpdateRegistration, WeakThis = AsWeak(), this]()
+		{
+			if (!WeakThis.IsValid())
+			{
+				return;
+			}
+
+			// Re-register the port if required
+			if (bNeedsUpdateRegistration)
+			{
+				Register();
+			}
+
+			bPendingRegistration = false;
 
-	OnPortUpdated.Broadcast();
+			OnPortUpdated.Broadcast();
+		});
 }
 
 bool FDMXInputPort::IsRegistered() const
diff --git a/Engine/Plugins/VirtualProduction/DMX/DMXProtocol/Source/DMXProtocol/Private/IO/DMXOutputPort.cpp b/Engine/Plugins/VirtualProduction/DMX/DMXProtocol/Source/DMXProtocol/Private/IO/DMXOutputPort.cpp
--- a/Engine/Plugins/VirtualProduction/DMX/DMXProtocol/Source/DMXProtocol/Private/IO/DMXOutputPort.cpp
+++ b/Engine/Plugins/VirtualProduction/DMX/DMXProtocol/Source/DMXProtocol/Private/IO/DMXOutputPort.cpp
@@ -2,6 +2,7 @@
 
 #include "IO/DMXOutputPort.h"
 
+#include "Async/Async.h"
 #include "DMXProtocolLog.h"
 #include "DMXProtocolSettings.h"
 #include "DMXProtocolUtils.h"
@@ -54,6 +55,7 @@ namespace UE::DMX::DMXOutputPort::Private::ConsoleCommands
 			FDMXOutputPortConfig PortConfig = OutputPort->MakeOutputPortConfig();
 			UE_LOG(LogDMXProtocol, Log, TEXT("Protocol: %s"), *PortConfig.GetProtocolName().ToString());
 			UE_LOG(LogDMXProtocol, Log, TEXT("Communication Type: %i"), static_cast<int32>(PortConfig.GetCommunicationType()));
+			UE_LOG(LogDMXProtocol, Log, TEXT("TimeToLive: %i"), PortConfig.GetTimeToLive());
 			UE_LOG(LogDMXProtocol, Log, TEXT("Device Address: %s"), *PortConfig.GetDeviceAddress());
 
 			const TArray<FDMXOutputPortDestinationAddress> DestinationAddresses = PortConfig.GetDestinationAddresses();
@@ -219,6 +221,36 @@ namespace UE::DMX::DMXOutputPort::Private::ConsoleCommands
 			})
 	);
 
+	static FAutoConsoleCommand GDMXSetTimeToLiveCommand(
+		TEXT("DMX.SetOutputPortTimeToLive"),
+		TEXT("DMX.SetOutputPortTimeToLive [PortName][TimeToLive (0-255, Default 1)]. Should usually not be changed. Sets how many hops a multicast DMX packet can take before being discarded. Example: DMX.SetOutputPortTimeToLive MyOutputPort 1"),
+		FConsoleCommandWithArgsDelegate::CreateStatic(
+			[](const TArray<FString>& Args)
+			{
+				constexpr int32 MinNumExpectedArgs = 2;
+				if (!VerifyConsoleCommandArguments(TEXT("DMX.SetOutputPortTimeToLive"), MinNumExpectedArgs, Args))
+				{
+					return;
+				}
+
+				const FString& PortName = Args[0];
+				const FString& TTLValueString = Args[1];
+
+				uint8 TTLValue;
+				if (LexTryParseString<uint8>(TTLValue, *TTLValueString))
+				{
+					const FDMXOutputPortSharedPtr OutputPort = FindOutputPortByName(PortName);
+
+					if (OutputPort.IsValid())
+					{
+						if (const IDMXProtocolPtr Protocol = OutputPort->GetProtocol())
+						{
+							DMX_OVERRIDE_OUTPUTPORT_VAR(TimeToLive, PortName, TTLValue);
+						}
+					}
+				}
+			})
+	);
 
 	static FAutoConsoleCommand GDMXSetOutputPortOutputPortDeviceAddressCommand(
 		TEXT("DMX.SetOutputPortDeviceAddress"),
@@ -645,7 +677,8 @@ void FDMXOutputPort::UpdateFromConfig(FDMXOutputPortConfig& InOutOutputPortConfi
 			DestinationAddresses == OutputPortConfig.GetDestinationAddresses() &&
 			CommunicationType == OutputPortConfig.GetCommunicationType() &&
 			Priority == OutputPortConfig.GetPriority() &&
-			DelaySeconds == OutputPortConfig.GetDelay() * OutputPortConfig.GetDelayFrameRate().AsDecimal())
+			DelaySeconds == OutputPortConfig.GetDelay() * OutputPortConfig.GetDelayFrameRate().AsDecimal() &&
+			TimeToLive == OutputPortConfig.GetTimeToLive())
 		{
 			return false;
 		}	
@@ -670,6 +703,7 @@ void FDMXOutputPort::UpdateFromConfig(FDMXOutputPortConfig& InOutOutputPortConfi
 	PortGuid = ConfigPortGuid;
 
 	CommunicationType = OutputPortConfig.GetCommunicationType();
+	TimeToLive = OutputPortConfig.GetTimeToLive();
 	bAutoCompleteDeviceAddressEnabled = OutputPortConfig.IsAutoCompleteDeviceAddressEnabled();
 	AutoCompleteDeviceAddress = OutputPortConfig.GetAutoCompleteDeviceAddress();
 	DeviceAddress = OutputPortConfig.GetDeviceAddress();
@@ -684,22 +718,30 @@ void FDMXOutputPort::UpdateFromConfig(FDMXOutputPortConfig& InOutOutputPortConfi
 
 	CommunicationDeterminator.SetLoopbackToEngine(OutputPortConfig.NeedsLoopbackToEngine());
 
-	const bool bValidDeviceAddress = FDMXProtocolUtils::GetLocalNetworkInterfaceCardIPs().ContainsByPredicate([this](const TSharedPtr<FString>& IPAddress)
-		{
-			return *IPAddress == DeviceAddress;
-		});
-	if (!bValidDeviceAddress && !IsRunningCookCommandlet())
+	// Only update on the next tick to give the OS time to fully release previous sockets
+	if (bPendingRegistration)
 	{
-		UE_LOG(LogDMXProtocol, Warning, TEXT("Cannot register Output Port %s. Device Address '%s' does not exist on the local machine."), *PortName, *DeviceAddress);
+		return;
 	}
+	bPendingRegistration = true;
 
-	// Re-register the port if required
-	if (bNeedsUpdateRegistration)
-	{
-		Register();
-	}
+	AsyncTask(ENamedThreads::GameThread, [bNeedsUpdateRegistration, WeakThis = AsWeak(), this]()
+		{
+			if (!WeakThis.IsValid())
+			{
+				return;
+			}
+
+			// Re-register the port if required
+			if (bNeedsUpdateRegistration)
+			{
+				Register();
+			}
 
-	OnPortUpdated.Broadcast();
+			bPendingRegistration = false;
+
+			OnPortUpdated.Broadcast();
+		});
 }
 
 const FGuid& FDMXOutputPort::GetPortGuid() const
diff --git a/Engine/Plugins/VirtualProduction/DMX/DMXProtocol/Source/DMXProtocol/Private/IO/DMXOutputPortConfig.cpp b/Engine/Plugins/VirtualProduction/DMX/DMXProtocol/Source/DMXProtocol/Private/IO/DMXOutputPortConfig.cpp
--- a/Engine/Plugins/VirtualProduction/DMX/DMXProtocol/Source/DMXProtocol/Private/IO/DMXOutputPortConfig.cpp
+++ b/Engine/Plugins/VirtualProduction/DMX/DMXProtocol/Source/DMXProtocol/Private/IO/DMXOutputPortConfig.cpp
@@ -14,6 +14,7 @@ FDMXOutputPortConfigParams::FDMXOutputPortConfigParams(const FDMXOutputPortConfi
 	: PortName(OutputPortConfig.GetPortName())
 	, ProtocolName(OutputPortConfig.GetProtocolName())
 	, CommunicationType(OutputPortConfig.GetCommunicationType())
+	, TimeToLive(OutputPortConfig.GetTimeToLive())
 	, bAutoCompleteDeviceAddressEnabled(OutputPortConfig.IsAutoCompleteDeviceAddressEnabled())
 	, AutoCompleteDeviceAddress(OutputPortConfig.GetAutoCompleteDeviceAddress())
 	, DeviceAddress(OutputPortConfig.GetDeviceAddress())
@@ -48,6 +49,7 @@ FDMXOutputPortConfig::FDMXOutputPortConfig(const FGuid& InPortGuid, const FDMXOu
 	: PortName(InitializationData.PortName)
 	, ProtocolName(InitializationData.ProtocolName)
 	, CommunicationType(InitializationData.CommunicationType)
+	, TimeToLive(InitializationData.TimeToLive)
 	, bAutoCompleteDeviceAddressEnabled(InitializationData.bAutoCompleteDeviceAddressEnabled)
 	, AutoCompleteDeviceAddress(InitializationData.AutoCompleteDeviceAddress)
 	, DeviceAddress(InitializationData.DeviceAddress)
diff --git a/Engine/Plugins/VirtualProduction/DMX/DMXProtocol/Source/DMXProtocol/Public/IO/DMXInputPort.h b/Engine/Plugins/VirtualProduction/DMX/DMXProtocol/Source/DMXProtocol/Public/IO/DMXInputPort.h
--- a/Engine/Plugins/VirtualProduction/DMX/DMXProtocol/Source/DMXProtocol/Public/IO/DMXInputPort.h
+++ b/Engine/Plugins/VirtualProduction/DMX/DMXProtocol/Source/DMXProtocol/Public/IO/DMXInputPort.h
@@ -147,6 +147,9 @@ private:
 	/** Mutex access for the RawListeners array */
 	FCriticalSection AccessRawListenersMutex;
 
+	/** True while registration is pending */
+	bool bPendingRegistration = false;
+
 	/** True if the port is registered with it its protocol */
 	TAtomic<bool> bRegistered = false;
 
diff --git a/Engine/Plugins/VirtualProduction/DMX/DMXProtocol/Source/DMXProtocol/Public/IO/DMXOutputPort.h b/Engine/Plugins/VirtualProduction/DMX/DMXProtocol/Source/DMXProtocol/Public/IO/DMXOutputPort.h
--- a/Engine/Plugins/VirtualProduction/DMX/DMXProtocol/Source/DMXProtocol/Public/IO/DMXOutputPort.h
+++ b/Engine/Plugins/VirtualProduction/DMX/DMXProtocol/Source/DMXProtocol/Public/IO/DMXOutputPort.h
@@ -149,6 +149,9 @@ public:
 	UE_DEPRECATED(4.27, "Use GameThreadGetDMXSignal instead. GameThreadGetDMXSignalFromRemoteUniverse only exists to support deprecated blueprint nodes.")
 	bool GameThreadGetDMXSignalFromRemoteUniverse(FDMXSignalSharedPtr& OutDMXSignal, int32 RemoteUniverseID, bool bEvenIfNotLoopbackToEngine);
 
+	/** Returns the multicast Time-to-Live (TTL), only used when the port is set to multicast */
+	uint8 GetTimeToLive() const { return TimeToLive; }
+
 	/** Returns the output port's delay in seconds */
 	FORCEINLINE double GetDelaySeconds() const { return DelaySeconds; }
 
@@ -210,12 +213,18 @@ private:
 	/** Mutex access to members required to send DMX  */
 	mutable FCriticalSection SendDMXMutex;
 
+	/** Multicast Time-to-Live (TTL), only used when the port is set to multicast */
+	uint8 TimeToLive = 1;
+
 	/** Priority on which packets are being sent */
 	int32 Priority = 0;
 
 	/** Map of raw Inputs */
 	TSet<TSharedRef<FDMXRawListener>> RawListeners;
 
+	/** True while registration is pending */
+	bool bPendingRegistration = false;
+
 	/** True if the port is registered with its protocol */
 	TAtomic<bool> bRegistered = false;
 
diff --git a/Engine/Plugins/VirtualProduction/DMX/DMXProtocol/Source/DMXProtocol/Public/IO/DMXOutputPortConfig.h b/Engine/Plugins/VirtualProduction/DMX/DMXProtocol/Source/DMXProtocol/Public/IO/DMXOutputPortConfig.h
--- a/Engine/Plugins/VirtualProduction/DMX/DMXProtocol/Source/DMXProtocol/Public/IO/DMXOutputPortConfig.h
+++ b/Engine/Plugins/VirtualProduction/DMX/DMXProtocol/Source/DMXProtocol/Public/IO/DMXOutputPortConfig.h
@@ -46,6 +46,7 @@ struct DMXPROTOCOL_API FDMXOutputPortConfigParams
 	FString PortName;
 	FName ProtocolName;
 	EDMXCommunicationType CommunicationType;
+	uint8 TimeToLive;
 	bool bAutoCompleteDeviceAddressEnabled;
 	FString AutoCompleteDeviceAddress;
 	FString DeviceAddress;
@@ -91,8 +92,9 @@ public:
 	FORCEINLINE const FString& GetPortName() const { return PortName; }
 	FORCEINLINE const FName& GetProtocolName() const { return ProtocolName; }
 	FORCEINLINE EDMXCommunicationType GetCommunicationType() const { return CommunicationType; }
-	FORCEINLINE bool IsAutoCompleteDeviceAddressEnabled() const { return bAutoCompleteDeviceAddressEnabled; }
-	FORCEINLINE FString GetAutoCompleteDeviceAddress() const { return AutoCompleteDeviceAddress; }
+	uint8 GetTimeToLive() const { return TimeToLive; }
+	bool IsAutoCompleteDeviceAddressEnabled() const { return bAutoCompleteDeviceAddressEnabled; }
+	FString GetAutoCompleteDeviceAddress() const { return AutoCompleteDeviceAddress; }
 	FString GetDeviceAddress() const;
 	FORCEINLINE const TArray<FDMXOutputPortDestinationAddress>& GetDestinationAddresses() const { return DestinationAddresses; }
 	FORCEINLINE bool NeedsLoopbackToEngine() const { return bLoopbackToEngine; }
@@ -141,6 +143,10 @@ protected:
 	UPROPERTY(Config, BlueprintReadWrite, EditDefaultsOnly, Category = "Port Config")
 	EDMXCommunicationType CommunicationType = EDMXCommunicationType::InternalOnly;
 
+	/** Sets how many hops multicast DMX packets can take before being discarded. In common network setups this should be set to 1. */
+	UPROPERTY()
+	uint8 TimeToLive = 1;
+
 	/** Enables 'Auto Complete Device Address', hidden via customization - EditConditionInlineToggle doesn't support Config */
 	UPROPERTY(Config, BlueprintReadOnly, EditDefaultsOnly, Category = "Port Config")
 	bool bAutoCompleteDeviceAddressEnabled = false;
diff --git a/Engine/Plugins/VirtualProduction/DMX/DMXProtocol/Source/DMXProtocolSACN/Private/DMXProtocolSACN.cpp b/Engine/Plugins/VirtualProduction/DMX/DMXProtocol/Source/DMXProtocolSACN/Private/DMXProtocolSACN.cpp
--- a/Engine/Plugins/VirtualProduction/DMX/DMXProtocol/Source/DMXProtocolSACN/Private/DMXProtocolSACN.cpp
+++ b/Engine/Plugins/VirtualProduction/DMX/DMXProtocol/Source/DMXProtocolSACN/Private/DMXProtocolSACN.cpp
@@ -147,7 +147,8 @@ TArray<TSharedPtr<IDMXSender>> FDMXProtocolSACN::RegisterOutputPort(const TShare
 	check(!CachedOutputPorts.Contains(OutputPort));
 
 	const FString& NetworkInterfaceAddress = OutputPort->GetDeviceAddress();
-	EDMXCommunicationType CommunicationType = OutputPort->GetCommunicationType();
+	const EDMXCommunicationType CommunicationType = OutputPort->GetCommunicationType();
+	const uint8 TTL = OutputPort->GetTimeToLive();
 
 	// Try to use an existing receiver or create a new one
 	TArray<TSharedPtr<IDMXSender>> NewSenders;
@@ -156,7 +157,7 @@ TArray<TSharedPtr<IDMXSender>> FDMXProtocolSACN::RegisterOutputPort(const TShare
 		TSharedPtr<FDMXProtocolSACNSender> NewSender = FindExistingMulticastSender(NetworkInterfaceAddress);
 		if (!NewSender.IsValid())
 		{
-			NewSender = FDMXProtocolSACNSender::TryCreateMulticastSender(SharedThis(this), NetworkInterfaceAddress);
+			NewSender = FDMXProtocolSACNSender::TryCreateMulticastSender(SharedThis(this), NetworkInterfaceAddress, TTL);
 		}
 
 		if (NewSender.IsValid())
diff --git a/Engine/Plugins/VirtualProduction/DMX/DMXProtocol/Source/DMXProtocolSACN/Private/DMXProtocolSACNReceiver.cpp b/Engine/Plugins/VirtualProduction/DMX/DMXProtocol/Source/DMXProtocolSACN/Private/DMXProtocolSACNReceiver.cpp
--- a/Engine/Plugins/VirtualProduction/DMX/DMXProtocol/Source/DMXProtocolSACN/Private/DMXProtocolSACNReceiver.cpp
+++ b/Engine/Plugins/VirtualProduction/DMX/DMXProtocol/Source/DMXProtocolSACN/Private/DMXProtocolSACNReceiver.cpp
@@ -65,15 +65,20 @@ TSharedPtr<FDMXProtocolSACNReceiver> FDMXProtocolSACNReceiver::TryCreate(const T
 		return nullptr;
 	}
 	const FIPv4Endpoint Endpoint(EndpointInternetAddr);
-	const FIPv4Endpoint AnyEndpoint(FIPv4Address::Any, ACN_PORT);
 
 	FSocket* NewListeningSocket = FUdpSocketBuilder(TEXT("UDPSACNListeningSocket"))
 		.AsBlocking()
 		.AsReusable()
-		.BoundToEndpoint(AnyEndpoint)
+#if PLATFORM_WINDOWS
+		.BoundToAddress(Endpoint.Address)
+#endif
+		.BoundToPort(ACN_PORT)
+#if PLATFORM_SUPPORTS_UDP_MULTICAST_GROUP
 		.WithMulticastLoopback()
 		.WithMulticastTtl(1)
 		.WithMulticastInterface(Endpoint.Address);
+#endif
+	;
 
 	if (!NewListeningSocket)
 	{
diff --git a/Engine/Plugins/VirtualProduction/DMX/DMXProtocol/Source/DMXProtocolSACN/Private/DMXProtocolSACNSender.cpp b/Engine/Plugins/VirtualProduction/DMX/DMXProtocol/Source/DMXProtocolSACN/Private/DMXProtocolSACNSender.cpp
--- a/Engine/Plugins/VirtualProduction/DMX/DMXProtocol/Source/DMXProtocolSACN/Private/DMXProtocolSACNSender.cpp
+++ b/Engine/Plugins/VirtualProduction/DMX/DMXProtocol/Source/DMXProtocolSACN/Private/DMXProtocolSACNSender.cpp
@@ -74,8 +74,7 @@ TSharedPtr<FDMXProtocolSACNSender> FDMXProtocolSACNSender::TryCreateUnicastSende
 		FUdpSocketBuilder(TEXT("UDPSACNUnicastSocket"))
 		.AsBlocking()
 		.AsReusable()
-		.BoundToEndpoint(NewNetworkInterfaceEndpoint)
-		.WithMulticastLoopback();
+		.BoundToEndpoint(NewNetworkInterfaceEndpoint);
 
 	if(!NewSocket)
 	{
@@ -99,9 +98,15 @@ TSharedPtr<FDMXProtocolSACNSender> FDMXProtocolSACNSender::TryCreateUnicastSende
 TSharedPtr<FDMXProtocolSACNSender> FDMXProtocolSACNSender::TryCreateMulticastSender(
 	const TSharedPtr<FDMXProtocolSACN,
 	ESPMode::ThreadSafe>& SACNProtocol,
-	const FString& InNetworkInterfaceIP)
+	const FString& InNetworkInterfaceIP, 
+	const uint8 TTL)
 {
 	// Try to create a socket
+#if !PLATFORM_SUPPORTS_UDP_MULTICAST_GROUP
+	UE_LOG(LogDMXProtocol, Warning, TEXT("Cannot create a multicast sACN sender. Platform doesn't support UDP multicast groups."));
+	return nullptr;
+#endif 
+
 	const TSharedPtr<FInternetAddr> NewNetworkInterfaceInternetAddr = FDMXProtocolUtils::CreateInternetAddr(InNetworkInterfaceIP, ACN_PORT);
 	if (!NewNetworkInterfaceInternetAddr.IsValid())
 	{
@@ -114,9 +119,12 @@ TSharedPtr<FDMXProtocolSACNSender> FDMXProtocolSACNSender::TryCreateMulticastSen
 		FUdpSocketBuilder(TEXT("UDPSACNMulticastSocket"))
         .AsBlocking()
         .AsReusable()
+#if PLATFORM_SUPPORTS_UDP_MULTICAST_GROUP
         .WithMulticastLoopback()
 		.WithMulticastInterface(NewNetworkInterfaceEndpoint.Address)
-        .WithMulticastTtl(1);
+        .WithMulticastTtl(TTL)
+#endif
+	;
 	
 	if(!NewSocket)
 	{
diff --git a/Engine/Plugins/VirtualProduction/DMX/DMXProtocol/Source/DMXProtocolSACN/Private/DMXProtocolSACNSender.h b/Engine/Plugins/VirtualProduction/DMX/DMXProtocol/Source/DMXProtocolSACN/Private/DMXProtocolSACNSender.h
--- a/Engine/Plugins/VirtualProduction/DMX/DMXProtocol/Source/DMXProtocolSACN/Private/DMXProtocolSACNSender.h
+++ b/Engine/Plugins/VirtualProduction/DMX/DMXProtocol/Source/DMXProtocolSACN/Private/DMXProtocolSACNSender.h
@@ -54,7 +54,7 @@ public:
 	 * @param SACNProtocol			The protocol instance that wants to create the sender
 	 * @param InNetworkInterfaceIP		The IP address of the network interface that is used to send data
 	 */
-	static TSharedPtr<FDMXProtocolSACNSender> TryCreateMulticastSender(const TSharedPtr<FDMXProtocolSACN, ESPMode::ThreadSafe>& SACNProtocol, const FString& InNetworkInterfaceIP);
+	static TSharedPtr<FDMXProtocolSACNSender> TryCreateMulticastSender(const TSharedPtr<FDMXProtocolSACN, ESPMode::ThreadSafe>& SACNProtocol, const FString& InNetworkInterfaceIP, const uint8 TTL);
 
 	// ~Begin IDMXSender Interface
 	virtual void SendDMXSignal(const FDMXSignalSharedRef& DMXSignal) override;
