genereate spelling typo fixes

#rnx

diff --git a/Engine/Plugins/Interchange/Runtime/Source/Parsers/GLTFCore/Public/GLTFMesh.h b/Engine/Plugins/Interchange/Runtime/Source/Parsers/GLTFCore/Public/GLTFMesh.h
--- a/Engine/Plugins/Interchange/Runtime/Source/Parsers/GLTFCore/Public/GLTFMesh.h
+++ b/Engine/Plugins/Interchange/Runtime/Source/Parsers/GLTFCore/Public/GLTFMesh.h
@@ -258,7 +258,7 @@ namespace GLTF
 		bool bIsValid = Primitives.FindByPredicate([](const FPrimitive& Prim) { return !Prim.IsValid(); }) == nullptr;
 
 		//if MorphTargetNames are not set, but mesh has morph targets (which is likely), this will generate isValid false;
-		//IsValidCache will have to be (re-)genereated post FAsset::GenerateNames (which is called at the end of GLTF::FReader) to overcome this issue.
+		//IsValidCache will have to be (re-)generated post FAsset::GenerateNames (which is called at the end of GLTF::FReader) to overcome this issue.
 
 		//Validate Morph Target Names and Morph Target Weights:
 		if ((MorphTargetNames.Num() > 0) && (MorphTargetWeights.Num() > 0))
diff --git a/Engine/Plugins/Mutable/Source/CustomizableObjectEditor/Private/MuCOE/CustomizableObjectEditorViewportClient.cpp b/Engine/Plugins/Mutable/Source/CustomizableObjectEditor/Private/MuCOE/CustomizableObjectEditorViewportClient.cpp
--- a/Engine/Plugins/Mutable/Source/CustomizableObjectEditor/Private/MuCOE/CustomizableObjectEditorViewportClient.cpp
+++ b/Engine/Plugins/Mutable/Source/CustomizableObjectEditor/Private/MuCOE/CustomizableObjectEditorViewportClient.cpp
@@ -413,7 +413,7 @@ void FCustomizableObjectEditorViewportClient::Draw(const FSceneView* View, FPrim
 
 void FCustomizableObjectEditorViewportClient::Draw(FViewport* InViewport, FCanvas* Canvas)
 {
-	// Defensive check to avoid unreal crashing inside render if the mesh is degenereated
+	// Defensive check to avoid unreal crashing inside render if the mesh is degenerated
 	for (const TPair<FName, TWeakObjectPtr<UDebugSkelMeshComponent>>& Key : SkeletalMeshComponents)
 	{
 		UDebugSkelMeshComponent* Component = Key.Value.Get();
diff --git a/Engine/Plugins/VirtualProduction/DMX/DMXProtocol/Source/DMXProtocolEditor/Private/DetailsCustomizations/Widgets/SDMXDelayEditWidget.cpp b/Engine/Plugins/VirtualProduction/DMX/DMXProtocol/Source/DMXProtocolEditor/Private/DetailsCustomizations/Widgets/SDMXDelayEditWidget.cpp
--- a/Engine/Plugins/VirtualProduction/DMX/DMXProtocol/Source/DMXProtocolEditor/Private/DetailsCustomizations/Widgets/SDMXDelayEditWidget.cpp
+++ b/Engine/Plugins/VirtualProduction/DMX/DMXProtocol/Source/DMXProtocolEditor/Private/DetailsCustomizations/Widgets/SDMXDelayEditWidget.cpp
@@ -72,7 +72,7 @@ void SDMXDelayEditWidget::Construct(const FArguments& InArgs)
 		.AutoWidth()
 		[
 			SNew(SComboButton)
-			.OnGetMenuContent(this, &SDMXDelayEditWidget::GenereateDelayTypeMenu)
+			.OnGetMenuContent(this, &SDMXDelayEditWidget::GenerateDelayTypeMenu)
 			.ButtonContent()
 			[
 				SNew(STextBlock)
@@ -104,7 +104,7 @@ FFrameRate SDMXDelayEditWidget::GetDelayFrameRate() const
 	return DelayFrameRate;
 }
 
-TSharedRef<SWidget> SDMXDelayEditWidget::GenereateDelayTypeMenu()
+TSharedRef<SWidget> SDMXDelayEditWidget::GenerateDelayTypeMenu()
 {
 	FMenuBuilder MenuBuilder(true, nullptr);
 
diff --git a/Engine/Plugins/VirtualProduction/DMX/DMXProtocol/Source/DMXProtocolEditor/Private/DetailsCustomizations/Widgets/SDMXDelayEditWidget.h b/Engine/Plugins/VirtualProduction/DMX/DMXProtocol/Source/DMXProtocolEditor/Private/DetailsCustomizations/Widgets/SDMXDelayEditWidget.h
--- a/Engine/Plugins/VirtualProduction/DMX/DMXProtocol/Source/DMXProtocolEditor/Private/DetailsCustomizations/Widgets/SDMXDelayEditWidget.h
+++ b/Engine/Plugins/VirtualProduction/DMX/DMXProtocol/Source/DMXProtocolEditor/Private/DetailsCustomizations/Widgets/SDMXDelayEditWidget.h
@@ -44,7 +44,7 @@ public:
 
 private:
 	/** Generates the menu to select a Type */
-	TSharedRef<SWidget> GenereateDelayTypeMenu();
+	TSharedRef<SWidget> GenerateDelayTypeMenu();
 
 	/** Called when the value text changed but was not commited */
 	bool OnVerifyValueTextChanged(const FText& InNewText, FText& OutErrorMessage);
diff --git a/Engine/Source/Programs/UnrealBuildTool/Configuration/UEBuildModuleCPP.cs b/Engine/Source/Programs/UnrealBuildTool/Configuration/UEBuildModuleCPP.cs
--- a/Engine/Source/Programs/UnrealBuildTool/Configuration/UEBuildModuleCPP.cs
+++ b/Engine/Source/Programs/UnrealBuildTool/Configuration/UEBuildModuleCPP.cs
@@ -627,7 +627,7 @@ namespace UnrealBuildTool
 				// compiled the same way. This can cause false positive in some situations, for example with address
 				// sanitize enabled (see: [REDACTED_PROFILE_URL]/protobuf/issues/13115).
 				ProtoCompileEnvironment.bDisableDynamicAnalysis = true;
-				// Disable warnings that are common in genereated proto cpp
+				// Disable warnings that are common in generated proto cpp
 				ProtoCompileEnvironment.CppCompileWarnings.DeprecationWarningLevel = WarningLevel.Off;
 				ProtoCompileEnvironment.CppCompileWarnings.ShortenSizeTToIntWarningLevel = WarningLevel.Off;
 				ProtoCompileEnvironment.CppCompileWarnings.UnusedParameterWarningLevel = WarningLevel.Off;
